Fixed UAF Control Rig trait mapped controls being initialized with random memory, due to controls not having associated property to initialize the latent pin memory

diff --git a/Engine/Plugins/Experimental/UAF/UAFControlRig/Source/UAFControlRig/Internal/ControlRigTrait.h b/Engine/Plugins/Experimental/UAF/UAFControlRig/Source/UAFControlRig/Internal/ControlRigTrait.h
--- a/Engine/Plugins/Experimental/UAF/UAFControlRig/Source/UAFControlRig/Internal/ControlRigTrait.h
+++ b/Engine/Plugins/Experimental/UAF/UAFControlRig/Source/UAFControlRig/Internal/ControlRigTrait.h
@@ -163,6 +163,8 @@ struct FControlRigTraitSharedData : public FAnimNextTraitSharedData
 
 	static void DestructLatentProperties(const UE::UAF::FTraitBinding& Binding);
 
+	static void InitializeControlLatentPinDefaultValue(ERigControlType InControlType, uint8* InTargetLatentPinMemory);
+
 #if WITH_EDITOR
 	UAFCONTROLRIG_API USkeleton* GetPreviewSkeleton() const;
 #endif // WITH_EDITOR
diff --git a/Engine/Plugins/Experimental/UAF/UAFControlRig/Source/UAFControlRig/Private/ControlRigTrait.cpp b/Engine/Plugins/Experimental/UAF/UAFControlRig/Source/UAFControlRig/Private/ControlRigTrait.cpp
--- a/Engine/Plugins/Experimental/UAF/UAFControlRig/Source/UAFControlRig/Private/ControlRigTrait.cpp
+++ b/Engine/Plugins/Experimental/UAF/UAFControlRig/Source/UAFControlRig/Private/ControlRigTrait.cpp
@@ -42,12 +42,19 @@ void FControlRigTraitSharedData::ConstructLatentProperties(const UE::UAF::FTrait
 		for (const FControlRigVariableMappings::FCustomPropertyData& Mapping : Mappings)
 		{
 			const FProperty* Property = Mapping.Property;
-			const uint8* LatentPinMemory = Mapping.SourceMemory;
-			if (Property && LatentPinMemory)
+			if (const uint8* LatentPinMemory = Mapping.SourceMemory)
 			{
-				// Init latent memory to default value
-				uint8* MutableMemory = const_cast<uint8*>(LatentPinMemory);
-				Property->InitializeValue(MutableMemory);
+				if (Property)
+				{
+					// Init latent memory to default value
+					uint8* MutableMemory = const_cast<uint8*>(LatentPinMemory);
+					Property->InitializeValue(MutableMemory);
+				}
+				else if (Mapping.Type == FControlRigVariableMappings::FCustomPropertyData::ECustomPropertyType::Control)
+				{
+					uint8* MutableMemory = const_cast<uint8*>(LatentPinMemory);
+					InitializeControlLatentPinDefaultValue(Mapping.ControlType, MutableMemory);
+				}
 			}
 		}
 	}
@@ -76,6 +83,68 @@ void FControlRigTraitSharedData::DestructLatentProperties(const UE::UAF::FTraitB
 	}
 }
 
+// Unlike mapped variables, controls does not come with an associated property, 
+// so I default init the latent property memory until the graph provides a valid value
+// This avoids setting control initial values with random memory
+void FControlRigTraitSharedData::InitializeControlLatentPinDefaultValue(ERigControlType InControlType, uint8* InTargetLatentPinMemory)
+{
+	check(InTargetLatentPinMemory != nullptr);
+
+	switch (InControlType)
+	{
+		case ERigControlType::Bool:
+		{
+			bool& ValuePtr = *(bool*)InTargetLatentPinMemory;
+			ValuePtr = false;
+			break;
+		}
+		case ERigControlType::Float:
+		case ERigControlType::ScaleFloat:
+		{
+			float& ValuePtr = *(float*)InTargetLatentPinMemory;
+			ValuePtr = 0.f;
+			break;
+		}
+		case ERigControlType::Integer:
+		{
+			int32& ValuePtr = *(int32*)InTargetLatentPinMemory;
+			ValuePtr = 0;
+			break;
+		}
+		case ERigControlType::Vector2D:
+		{
+			FVector2D& ValuePtr = *(FVector2D*)InTargetLatentPinMemory;
+			ValuePtr = FVector2D::ZeroVector;
+			break;
+		}
+		case ERigControlType::Position:
+		case ERigControlType::Scale:
+		{
+			FVector& ValuePtr = *(FVector*)InTargetLatentPinMemory;
+			ValuePtr = FVector::ZeroVector;
+			break;
+		}
+		case ERigControlType::Rotator:
+		{
+			FRotator& ValuePtr = *(FRotator*)InTargetLatentPinMemory;
+			ValuePtr = FRotator::ZeroRotator;
+			break;
+		}
+		case ERigControlType::Transform:
+		case ERigControlType::TransformNoScale:
+		case ERigControlType::EulerTransform:
+		{
+			FTransform& ValuePtr = *(FTransform*)InTargetLatentPinMemory;
+			ValuePtr = FTransform::Identity;
+			break;
+		}
+		default:
+		{
+			checkNoEntry();
+		}
+	}
+}
+
 #if WITH_EDITOR
 USkeleton* FControlRigTraitSharedData::GetPreviewSkeleton() const
 {
