UAF Control Rig Trait : Fixed latent pin size error for position and scale, as it was using the type of the target control, rather than the type of the source pin

diff --git a/Engine/Plugins/Experimental/UAF/UAFControlRig/Source/UAFControlRig/Private/ControlRigTrait.cpp b/Engine/Plugins/Experimental/UAF/UAFControlRig/Source/UAFControlRig/Private/ControlRigTrait.cpp
--- a/Engine/Plugins/Experimental/UAF/UAFControlRig/Source/UAFControlRig/Private/ControlRigTrait.cpp
+++ b/Engine/Plugins/Experimental/UAF/UAFControlRig/Source/UAFControlRig/Private/ControlRigTrait.cpp
@@ -973,9 +973,12 @@ bool FControlRigTrait::GetControlSizeAndAlignment(ERigControlType ControlType, u
 		case ERigControlType::Position:
 		case ERigControlType::Scale:
 		{
-			PropertySize = sizeof(FVector3f);
-			PropertyAlignment = PropertySize;
-			return true;
+			if (const UScriptStruct* Struct = TBaseStructure<FVector>::Get())
+			{
+				PropertySize = Struct->GetStructureSize();
+				PropertyAlignment = Struct->GetMinAlignment();
+				return true;
+			}
 		}
 		case ERigControlType::Rotator:
 		{
