UAF Control Rig: Refactor the control rig trait and evaluation task to be more modular so that the control rig evaluation task can be used outside of UAF in Sequencer animation mixer evaluation.

diff --git a/Engine/Plugins/Animation/ControlRig/Source/ControlRig/Private/Tools/ControlRigVariableMappings.cpp b/Engine/Plugins/Animation/ControlRig/Source/ControlRig/Private/Tools/ControlRigVariableMappings.cpp
--- a/Engine/Plugins/Animation/ControlRig/Source/ControlRig/Private/Tools/ControlRigVariableMappings.cpp
+++ b/Engine/Plugins/Animation/ControlRig/Source/ControlRig/Private/Tools/ControlRigVariableMappings.cpp
@@ -209,7 +209,7 @@ void FControlRigVariableMappings::PropagateCustomInputProperties(UControlRig* In
 	}
 }
 
-void FControlRigVariableMappings::UpdateCurveInputs(UControlRig* InControlRig, const TMap<FName, FName> &InputMapping, const FBlendedCurve& InCurveData)
+void FControlRigVariableMappings::UpdateCurveInputs(UControlRig* InControlRig, const TMap<FName, FName> &InputMapping, const FBlendedCurve& InCurveData) const
 {
 	DECLARE_SCOPE_HIERARCHICAL_COUNTER_FUNC()
 
@@ -232,7 +232,7 @@ void FControlRigVariableMappings::UpdateCurveInputs(UControlRig* InControlRig, c
 	}
 }
 
-void FControlRigVariableMappings::UpdateCurveOutputs(UControlRig* InControlRig, const TMap<FName, FName> &OutputMapping, FBlendedCurve& InCurveData)
+void FControlRigVariableMappings::UpdateCurveOutputs(UControlRig* InControlRig, const TMap<FName, FName> &OutputMapping, FBlendedCurve& InCurveData) const
 {
 	DECLARE_SCOPE_HIERARCHICAL_COUNTER_FUNC()
 
diff --git a/Engine/Plugins/Animation/ControlRig/Source/ControlRig/Public/Tools/ControlRigVariableMappings.h b/Engine/Plugins/Animation/ControlRig/Source/ControlRig/Public/Tools/ControlRigVariableMappings.h
--- a/Engine/Plugins/Animation/ControlRig/Source/ControlRig/Public/Tools/ControlRigVariableMappings.h
+++ b/Engine/Plugins/Animation/ControlRig/Source/ControlRig/Public/Tools/ControlRigVariableMappings.h
@@ -111,9 +111,9 @@ struct FControlRigVariableMappings
 	UE_API void PropagateInputProperties(const UObject* InSourceInstance, UControlRig* InTargetInstance, TArray<FName>& InDestPropertyNames);
 	UE_API void PropagateCustomInputProperties(UControlRig* InTargetControlRig);
 
-	UE_API void UpdateCurveInputs(UControlRig* InControlRig, const TMap<FName, FName>& InputMapping, const FBlendedCurve& CurveData);
+	UE_API void UpdateCurveInputs(UControlRig* InControlRig, const TMap<FName, FName>& InputMapping, const FBlendedCurve& CurveData) const;
 
-	UE_API void UpdateCurveOutputs(UControlRig* InControlRig, const TMap<FName, FName>& OutputMapping, FBlendedCurve& InCurveData);
+	UE_API void UpdateCurveOutputs(UControlRig* InControlRig, const TMap<FName, FName>& OutputMapping, FBlendedCurve& InCurveData) const;
 
 	void ResetCurvesInputToControlCache()
 	{
diff --git a/Engine/Plugins/Experimental/UAF/UAFControlRig/Source/UAFControlRig/Internal/AnimNextControlRigHierarchyMappings.h b/Engine/Plugins/Experimental/UAF/UAFControlRig/Source/UAFControlRig/Internal/AnimNextControlRigHierarchyMappings.h
--- a/Engine/Plugins/Experimental/UAF/UAFControlRig/Source/UAFControlRig/Internal/AnimNextControlRigHierarchyMappings.h
+++ b/Engine/Plugins/Experimental/UAF/UAFControlRig/Source/UAFControlRig/Internal/AnimNextControlRigHierarchyMappings.h
@@ -64,7 +64,7 @@ struct UAFCONTROLRIG_API FAnimNextControlRigHierarchyMappings
 		, UE::UAF::FKeyframeState& InOutput
 		, const FControlRigIOSettings& InInputSettings
 		, const FControlRigIOSettings& InOutputSettings
-		, TWeakObjectPtr<UNodeMappingContainer>& InNodeMappingContainer
+		, const TWeakObjectPtr<UNodeMappingContainer>& InNodeMappingContainer
 		, bool bInExecute
 		, bool bInTransferInputPose
 		, bool bInResetInputPoseToInitial
@@ -74,7 +74,7 @@ struct UAFCONTROLRIG_API FAnimNextControlRigHierarchyMappings
 	void UpdateOutput(UControlRig* ControlRig
 		, UE::UAF::FKeyframeState& InOutput
 		, const FControlRigIOSettings& InOutputSettings
-		, TWeakObjectPtr<UNodeMappingContainer>& InNodeMappingContainer
+		, const TWeakObjectPtr<UNodeMappingContainer>& InNodeMappingContainer
 		, bool bInExecute
 		, bool bInTransferPoseInGlobalSpace);
 
diff --git a/Engine/Plugins/Experimental/UAF/UAFControlRig/Source/UAFControlRig/Internal/ControlRigTrait.h b/Engine/Plugins/Experimental/UAF/UAFControlRig/Source/UAFControlRig/Internal/ControlRigTrait.h
--- a/Engine/Plugins/Experimental/UAF/UAFControlRig/Source/UAFControlRig/Internal/ControlRigTrait.h
+++ b/Engine/Plugins/Experimental/UAF/UAFControlRig/Source/UAFControlRig/Internal/ControlRigTrait.h
@@ -14,6 +14,7 @@
 #include "AnimNextControlRigHierarchyMappings.h"
 #include "Tools/ControlRigVariableMappings.h"
 #include "Tools/ControlRigIOSettings.h"
+#include "ControlRigTask.h"
 #if WITH_EDITOR
 #include "ControlRigIOMapping.h"
 #endif
@@ -22,19 +23,7 @@
 
 class USkeletalMeshComponent;
 struct FRigVMExternalVariable;
-
-USTRUCT()
-struct FControlRigEventName
-{
-	GENERATED_BODY()
-
-	FControlRigEventName()
-	: EventName(NAME_None)
-	{}
-
-	UPROPERTY(EditAnywhere, Category = Links)
-	FName EventName;	
-};
+struct FAnimNextControlRigTask;
 
 USTRUCT()
 struct FControlRigExposedProperty
@@ -188,32 +177,14 @@ struct FControlRigTrait : FBaseTrait, IEvaluate, IUpdate, IHierarchy, IGarbageCo
 	{
 		FTraitPtr Input;
 
-		FRigVMDrawInterface* DebugDrawInterface = nullptr;
-		FTransform ComponentTransform;
-
-		TObjectPtr<USkeletalMeshComponent> SkeletalMeshComponent;
+		TSharedPtr<FAnimNextControlRigTask> EvaluationTask;
 
 		TObjectPtr<UControlRig> ControlRig;
 
-		ControlRig::FAnimNextControlRigHierarchyMappings ControlRigHierarchyMappings;
-		FControlRigVariableMappings ControlRigVariableMappings;
-		FControlRigVariableMappings::FCustomPropertyMappings PropertyMappings;
-
-		TWeakObjectPtr<UNodeMappingContainer> NodeMappingContainer;
-		FControlRigIOSettings InputSettings;
-		FControlRigIOSettings OutputSettings;
-
-		uint16 LastBonesSerialNumberForCacheBones = 0;
-		bool bControlRigRequiresInitialization = false;
-
 		FDelegateHandle OnObjectsReinstancedHandle;
 		FDelegateHandle OnInitializedHandle;
 
-		int32 LastLOD = INDEX_NONE;
-
-		bool bExecute = true;
-		bool bClearEventQueueRequired = false;
-		bool bUpdateInputOutputMapping = false;
+		FControlRigVariableMappings::FCustomPropertyMappings PropertyMappings;
 
 		void Construct(const FExecutionContext& Context, const FTraitBinding& Binding);
 		void Destruct(const FExecutionContext& Context, const FTraitBinding& Binding);
diff --git a/Engine/Plugins/Experimental/UAF/UAFControlRig/Source/UAFControlRig/Private/AnimNextControlRigHierarchyMappings.cpp b/Engine/Plugins/Experimental/UAF/UAFControlRig/Source/UAFControlRig/Private/AnimNextControlRigHierarchyMappings.cpp
--- a/Engine/Plugins/Experimental/UAF/UAFControlRig/Source/UAFControlRig/Private/AnimNextControlRigHierarchyMappings.cpp
+++ b/Engine/Plugins/Experimental/UAF/UAFControlRig/Source/UAFControlRig/Private/AnimNextControlRigHierarchyMappings.cpp
@@ -119,7 +119,7 @@ void FAnimNextControlRigHierarchyMappings::UpdateInput(UControlRig* ControlRig
 	, UE::UAF::FKeyframeState& InOutput
 	, const FControlRigIOSettings& InInputSettings
 	, const FControlRigIOSettings& InOutputSettings
-	, TWeakObjectPtr<UNodeMappingContainer>& InNodeMappingContainer
+	, const TWeakObjectPtr<UNodeMappingContainer>& InNodeMappingContainer
 	, bool bInExecute
 	, bool bInTransferInputPose
 	, bool bInResetInputPoseToInitial
@@ -199,7 +199,7 @@ void FAnimNextControlRigHierarchyMappings::UpdateInput(UControlRig* ControlRig
 void FAnimNextControlRigHierarchyMappings::UpdateOutput(UControlRig* ControlRig
 	, UE::UAF::FKeyframeState& InOutput
 	, const FControlRigIOSettings& InOutputSettings
-	, TWeakObjectPtr<UNodeMappingContainer>& InNodeMappingContainer
+	, const TWeakObjectPtr<UNodeMappingContainer>& InNodeMappingContainer
 	, bool bInExecute
 	, bool bInTransferPoseInGlobalSpace)
 {
diff --git a/Engine/Plugins/Experimental/UAF/UAFControlRig/Source/UAFControlRig/Private/ControlRigTask.cpp b/Engine/Plugins/Experimental/UAF/UAFControlRig/Source/UAFControlRig/Private/ControlRigTask.cpp
--- a/Engine/Plugins/Experimental/UAF/UAFControlRig/Source/UAFControlRig/Private/ControlRigTask.cpp
+++ b/Engine/Plugins/Experimental/UAF/UAFControlRig/Source/UAFControlRig/Private/ControlRigTask.cpp
@@ -6,13 +6,13 @@
 #include "EvaluationVM/KeyframeState.h"
 #include "TransformArrayOperations.h"
 #include "BoneContainer.h"
-#include "Components/SkeletalMeshComponent.h"
 #include "Units/Execution/RigUnit_BeginExecution.h"
 #include "Math/Transform.h"
 #include "GenerationTools.h"
 #include "Engine/SkeletalMesh.h"
 #include "Animation/AnimNodeBase.h"
 #include "Units/Execution/RigUnit_PrepareForExecution.h"
+#include "ControlRigObjectBinding.h"
 #if ENABLE_ANIM_DEBUG
 #include "AnimNode_ControlRigBase.h"
 #endif
@@ -21,14 +21,63 @@ static TAutoConsoleVariable<int32> CVarControlRigDisableExecutionAnimNext(TEXT("
 
 DEFINE_STAT(STAT_AnimNext_Task_ControlRig);
 
-FAnimNextControlRigTask FAnimNextControlRigTask::Make(const UE::UAF::FControlRigTrait::FSharedData* SharedData, UE::UAF::FControlRigTrait::FInstanceData* InstanceData)
+
+void FAnimNextControlRigTask::InitializeControlRig(const FControlRigVariableMappings::FCustomPropertyMappings* PropertyMappings)
+{
+	if (Params.ControlRig)
+	{
+		// Provide available properties to the construction event
+		if (PropertyMappings)
+		{
+			Params.ControlRigVariableMappings.InitializeCustomProperties(Params.ControlRig, *PropertyMappings);
+			Params.ControlRigVariableMappings.PropagateCustomInputProperties(Params.ControlRig);
+		}
+
+		Params.ControlRigHierarchyMappings.InitializeInstance();
+
+		if (Params.ControlRig->GetObjectBinding())
+		{
+			if (USkeletalMeshComponent* SkeletalMeshComponent = Cast<USkeletalMeshComponent>(FControlRigObjectBinding::GetBindableObject(Params.ControlRig->GetObjectBinding()->GetBoundObject())))
+			{
+				Params.ControlRig->GetDataSourceRegistry()->RegisterDataSource(UControlRig::OwnerComponent, SkeletalMeshComponent);
+			}
+		}
+
+		Params.ControlRig->Initialize(true);
+		Params.ControlRig->RequestInit();
+
+		Params.ControlRigHierarchyMappings.ResetRefPoseSetterHash();
+		Params.ControlRigVariableMappings.ResetCurvesInputToControlCache();
+		Params.ControlRigVariableMappings.CacheCurveMappings(Params.InputMapping, Params.OutputMapping, Params.ControlRig->GetHierarchy());
+
+		if (PropertyMappings)
+		{
+			// Re-init Custom Properties after construction, as new controls could be created and might have to be remapped
+			Params.ControlRigVariableMappings.InitializeCustomProperties(Params.ControlRig, *PropertyMappings);
+			Params.ControlRigVariableMappings.PropagateCustomInputProperties(Params.ControlRig);
+		}
+	}
+}
+
+void FAnimNextControlRigTask::RequestInit()
 {
-	FAnimNextControlRigTask Task;
-	Task.InstanceData = InstanceData;
-	Task.SharedData = SharedData;
-	return Task;
+	if (UControlRig* ControlRig = GetControlRig())
+	{
+		ControlRig->RequestInit();
+	}
+	Params.bControlRigRequiresInitialization = true;
+	Params.LastBonesSerialNumberForCacheBones = 0;
+	Params.ControlRigHierarchyMappings.InitializeInstance();
+	Params.ControlRigHierarchyMappings.ResetRefPoseSetterHash();
+	Params.bUpdateInputOutputMapping = true;
+}
+
+void FAnimNextControlRigTask::OnControlRigInitialized()
+{
+	Params.ControlRigHierarchyMappings.ResetRefPoseSetterHash();
 }
 
+
 void FAnimNextControlRigTask::Execute(UE::UAF::FEvaluationVM& VM) const
 {
 	SCOPE_CYCLE_COUNTER(STAT_AnimNext_Task_ControlRig);
@@ -44,32 +93,48 @@ void FAnimNextControlRigTask::Execute(UE::UAF::FEvaluationVM& VM) const
 
 	TUniquePtr<FKeyframeState> KeyframeOut;
 
-	// Try to get a Keyframe from the stack
-	if (!VM.PopValue(UE::UAF::KEYFRAME_STACK_NAME, KeyframeOut))
+	if (Params.bConsumesPreviousPose)
+	{
+		// Try to get a Keyframe from the stack
+		if (!VM.PopValue(KEYFRAME_STACK_NAME, KeyframeOut))
+		{
+			// no-op, will be caught below in validity check, here to prevent nodiscard error
+		}
+	}
+	else if (const TUniquePtr<FKeyframeState>* SourcePose = VM.PeekValue<TUniquePtr<FKeyframeState>>(KEYFRAME_STACK_NAME, 0))
+	{
+		if (SourcePose->IsValid())
+		{
+			KeyframeOut = MakeUnique<FKeyframeState>(VM.MakeUninitializedKeyframe(false));
+			KeyframeOut->Pose.CopyFrom((*SourcePose)->Pose);
+			KeyframeOut->Curves.CopyFrom((*SourcePose)->Curves);
+			KeyframeOut->Attributes.CopyFrom((*SourcePose)->Attributes);
+		}
+	}
+
+	if (!KeyframeOut.IsValid())
 	{
-		// if there is none, just create a reference frame
 		KeyframeOut = MakeUnique<FKeyframeState>(VM.MakeReferenceKeyframe(false));
 	}
 
-	ExecuteControlRig(VM, *KeyframeOut.Get(), InstanceData->ControlRig);
+	ExecuteControlRig(VM, *KeyframeOut.Get());
 
 	// Push our blended result back
 	VM.PushValue(KEYFRAME_STACK_NAME, MoveTemp(KeyframeOut));
 
-	InstanceData->LastLOD = VM.GetCurrentLOD();
+	LastLOD = VM.GetCurrentLOD();
 }
 
-void FAnimNextControlRigTask::ExecuteControlRig(UE::UAF::FEvaluationVM& VM, UE::UAF::FKeyframeState& KeyFrameState, UControlRig* ControlRig) const
+void FAnimNextControlRigTask::ExecuteControlRig(UE::UAF::FEvaluationVM& VM, UE::UAF::FKeyframeState& KeyFrameState) const
 {
 	DECLARE_SCOPE_HIERARCHICAL_COUNTER_FUNC()
 
-	if (ControlRig)
+	if (Params.ControlRig)
 	{
 		// Before we start modifying the RigHierarchy, we need to lock the rig to avoid corrupted state
-		// TODO : Is this really needed ? 
-		//FScopeLock LockRig(&ControlRig->GetEvaluateMutex());
+		FScopeLock LockRig(&Params.ControlRig->GetEvaluateMutex());
 
-		URigHierarchy* Hierarchy = ControlRig->GetHierarchy();
+		URigHierarchy* Hierarchy = Params.ControlRig->GetHierarchy();
 		if (Hierarchy == nullptr)
 		{
 			return;
@@ -84,52 +149,53 @@ void FAnimNextControlRigTask::ExecuteControlRig(UE::UAF::FEvaluationVM& VM, UE::
 
 		// temporarily give control rig access to the stack allocated attribute container
 		// control rig may have rig units that can add/get attributes to/from this container
-		UControlRig::FAnimAttributeContainerPtrScope AttributeScope(ControlRig, MeshAttributeContainer);
+		UControlRig::FAnimAttributeContainerPtrScope AttributeScope(Params.ControlRig, MeshAttributeContainer);
 
 		const int32 CurrentLOD = VM.GetCurrentLOD();
-		const bool bIsLODChange = InstanceData->LastLOD != CurrentLOD;
-		const bool bRunPrepareForExecution = bIsLODChange /* || (RefPoseBonesSerialNumber != LastRefPoseBonesSerialNumber)*/; // TODO : Detect RefPose changes
-		if (InstanceData->bUpdateInputOutputMapping || bRunPrepareForExecution)
+		const bool bIsLODChange = LastLOD != CurrentLOD;
+		const bool bRunPrepareForExecution = bIsLODChange || Params.bControlRigRequiresInitialization || Params.bUpdateInputOutputMapping;/* || (RefPoseBonesSerialNumber != LastRefPoseBonesSerialNumber)*/; // TODO : Detect RefPose changes
+		if (bRunPrepareForExecution)
 		{
-			if (ControlRig->IsConstructionModeEnabled() ||
-				(ControlRig->IsConstructionRequired() && (InstanceData->bUpdateInputOutputMapping || bRunPrepareForExecution)))
+			if (Params.ControlRig->IsConstructionModeEnabled() ||
+				(Params.ControlRig->IsConstructionRequired() && bRunPrepareForExecution))
 			{
-				//UpdateGetAssetUserDataDelegate(ControlRig);
-				ControlRig->Execute(FRigUnit_PrepareForExecution::EventName);
-				ControlRig->Execute(FRigUnit_PostPrepareForExecution::EventName);
+				//UpdateGetAssetUserDataDelegate(Params.ControlRig);
+				Params.ControlRig->Execute(FRigUnit_PrepareForExecution::EventName);
+				Params.ControlRig->Execute(FRigUnit_PostPrepareForExecution::EventName);
 			}
 
 			// UpdateInputOutputMappingIfRequired was done in CacheBones, but there is no AnimNext equivalent
-			InstanceData->ControlRigHierarchyMappings.UpdateInputOutputMappingIfRequired(ControlRig
-				, ControlRig->GetHierarchy()
+			Params.ControlRigHierarchyMappings.UpdateInputOutputMappingIfRequired(Params.ControlRig
+				, Params.ControlRig->GetHierarchy()
 				, RefPose
 				, CurrentLOD
 				, TArray<FBoneReference>()
 				, TArray<FBoneReference>()
-				, InstanceData->NodeMappingContainer
-				, SharedData->bTransferPoseInGlobalSpace
-				, SharedData->bResetInputPoseToInitial);
+				, nullptr
+				, Params.bTransferPoseInGlobalSpace
+				, Params.bResetInputPoseToInitial);
 
-			InstanceData->bUpdateInputOutputMapping = false;
+			Params.bUpdateInputOutputMapping = false;
+			Params.bControlRigRequiresInitialization = false;
 		}
 
-		if (!InstanceData->ControlRigHierarchyMappings.IsUpdateToDate(Hierarchy))
+		if (!Params.ControlRigHierarchyMappings.IsUpdateToDate(Hierarchy))
 		{
-			InstanceData->ControlRigHierarchyMappings.PerformUpdateToDate(ControlRig
+			Params.ControlRigHierarchyMappings.PerformUpdateToDate(Params.ControlRig
 				, Hierarchy
 				, RefPose
 				, CurrentLOD
-				, InstanceData->NodeMappingContainer
-				, SharedData->bTransferPoseInGlobalSpace
-				, SharedData->bResetInputPoseToInitial);
+				, nullptr
+				, Params.bTransferPoseInGlobalSpace
+				, Params.bResetInputPoseToInitial);
 		}
 
 		// first update input to the system
-		UpdateInput(VM, KeyFrameState, ControlRig, KeyFrameState);
+		UpdateInput(VM, KeyFrameState, KeyFrameState);
 
-		if (InstanceData->bExecute)
+		if (bExecute)
 		{
-			const TGuardValue<bool> ResetCurrentTransfromsAfterConstructionGuard = ControlRig->GetResetCurrentTransformsAfterConstructionGuard(true);
+			const TGuardValue<bool> ResetCurrentTransfromsAfterConstructionGuard = Params.ControlRig->GetResetCurrentTransformsAfterConstructionGuard(true);
 
 #if WITH_EDITOR
 			if (Hierarchy->IsTracingChanges())
@@ -139,43 +205,43 @@ void FAnimNextControlRigTask::ExecuteControlRig(UE::UAF::FEvaluationVM& VM, UE::
 #endif
 
 			// pick the event to run
-			if (SharedData->EventQueue.IsEmpty())
+			if (Params.EventQueue.IsEmpty())
 			{
-				if (InstanceData->bClearEventQueueRequired)
+				if (bClearEventQueueRequired)
 				{
-					ControlRig->SetEventQueue({ FRigUnit_BeginExecution::EventName });
-					InstanceData->bClearEventQueueRequired = false;
+					Params.ControlRig->SetEventQueue({ FRigUnit_BeginExecution::EventName });
+					bClearEventQueueRequired = false;
 				}
 			}
 			else
 			{
 				TArray<FName> EventNames;
-				Algo::Transform(SharedData->EventQueue, EventNames, [](const FControlRigEventName& InEventName)
+				Algo::Transform(Params.EventQueue, EventNames, [](const FControlRigEventName& InEventName)
 					{
 						return InEventName.EventName;
 					});
-				ControlRig->SetEventQueue(EventNames);
-				InstanceData->bClearEventQueueRequired = true;
+				Params.ControlRig->SetEventQueue(EventNames);
+				bClearEventQueueRequired = true;
 			}
 
-			if (ControlRig->IsAdditive())
+			if (Params.ControlRig->IsAdditive())
 			{
-				ControlRig->ClearPoseBeforeBackwardsSolve();
+				Params.ControlRig->ClearPoseBeforeBackwardsSolve();
 			}
 
 			// evaluate control rig
-			// UpdateGetAssetUserDataDelegate(ControlRig); // Removed in AnimNext, do we need an equivalent ?
-			ControlRig->Evaluate_AnyThread();
+			// UpdateGetAssetUserDataDelegate(Params.ControlRig); // Removed in AnimNext, do we need an equivalent ?
+			Params.ControlRig->Evaluate_AnyThread();
 
 #if ENABLE_ANIM_DEBUG
 #if UE_ENABLE_DEBUG_DRAWING
 			// When Control Rig is at editing time (in CR editor), draw instructions are consumed by ControlRigEditMode, so we need to skip drawing here.
-			const bool bShowDebug = (CVarAnimNodeControlRigDebug.GetValueOnAnyThread() == 1 && ControlRig->ExecutionType != ERigExecutionType::Editing);
+			const bool bShowDebug = (CVarAnimNodeControlRigDebug.GetValueOnAnyThread() == 1 && Params.ControlRig->ExecutionType != ERigExecutionType::Editing);
 			if (bShowDebug)
 			{
-				if (FRigVMDrawInterface* DebugDrawInterface = InstanceData->DebugDrawInterface)
+				if (Params.DebugDrawInterface)
 				{
-					QueueControlRigDrawInstructions(ControlRig, InstanceData->DebugDrawInterface, InstanceData->ComponentTransform);
+					QueueControlRigDrawInstructions(Params.DebugDrawInterface, Params.ComponentTransform);
 				}
 			}
 #endif
@@ -190,96 +256,103 @@ void FAnimNextControlRigTask::ExecuteControlRig(UE::UAF::FEvaluationVM& VM, UE::
 		}
 
 		// now update output
-		UpdateOutput(VM, KeyFrameState, ControlRig, KeyFrameState);
+		UpdateOutput(VM, KeyFrameState, KeyFrameState);
 		
 		// remap mesh bone index attributes back to stack container (LOD/compact bone indices)
 		UE::UAF::FGenerationTools::RemapAttributes(Pose, MeshAttributeContainer, KeyFrameState.Attributes);
 	}
 }
 
-void FAnimNextControlRigTask::UpdateInput(UE::UAF::FEvaluationVM& VM, UE::UAF::FKeyframeState& KeyFrameState, UControlRig* ControlRig, UE::UAF::FKeyframeState& InOutput) const
+void FAnimNextControlRigTask::UpdateInput(UE::UAF::FEvaluationVM& VM, UE::UAF::FKeyframeState& KeyFrameState, UE::UAF::FKeyframeState& InOutput) const
 {
 	DECLARE_SCOPE_HIERARCHICAL_COUNTER_FUNC()
 
-	if (!CanExecute(ControlRig))
+	if (!CanExecute())
 	{
 		return;
 	}
 
 	DECLARE_SCOPE_HIERARCHICAL_COUNTER_FUNC()
 
-	InstanceData->ControlRigHierarchyMappings.UpdateInput(ControlRig
+	Params.ControlRigHierarchyMappings.UpdateInput(Params.ControlRig
 		, InOutput
-		, InstanceData->InputSettings
-		, InstanceData->OutputSettings
-		, InstanceData->NodeMappingContainer
-		, InstanceData->bExecute
-		, SharedData->bTransferInputPose
-		, SharedData->bResetInputPoseToInitial
-		, SharedData->bTransferPoseInGlobalSpace
-		, SharedData->bTransferInputCurves);
-
-	InstanceData->ControlRigVariableMappings.UpdateCurveInputs(InstanceData->ControlRig, SharedData->InputMapping, InOutput.Curves);
+		, InputSettings
+		, OutputSettings
+		, nullptr
+		, bExecute
+		, Params.bTransferInputPose
+		, Params.bResetInputPoseToInitial
+		, Params.bTransferPoseInGlobalSpace
+		, Params.bTransferInputCurves);
+
+	Params.ControlRigVariableMappings.UpdateCurveInputs(Params.ControlRig, Params.InputMapping, InOutput.Curves);
 }
 
-void FAnimNextControlRigTask::UpdateOutput(UE::UAF::FEvaluationVM& VM, UE::UAF::FKeyframeState& KeyFrameState, UControlRig* ControlRig, UE::UAF::FKeyframeState& InOutput) const
+void FAnimNextControlRigTask::UpdateOutput(UE::UAF::FEvaluationVM& VM, UE::UAF::FKeyframeState& KeyFrameState, UE::UAF::FKeyframeState& InOutput) const
 {
 	DECLARE_SCOPE_HIERARCHICAL_COUNTER_FUNC()
 
-	if (!CanExecute(ControlRig))
+	if (!CanExecute())
 	{
 		return;
 	}
 
 	DECLARE_SCOPE_HIERARCHICAL_COUNTER_FUNC()
 
-	InstanceData->ControlRigHierarchyMappings.UpdateOutput(ControlRig
+	Params.ControlRigHierarchyMappings.UpdateOutput(Params.ControlRig
 		, InOutput
-		, InstanceData->OutputSettings
-		, InstanceData->NodeMappingContainer
-		, InstanceData->bExecute
-		, SharedData->bTransferPoseInGlobalSpace);
+		, OutputSettings
+		, nullptr
+		, bExecute
+		, Params.bTransferPoseInGlobalSpace);
 
-	InstanceData->ControlRigVariableMappings.UpdateCurveOutputs(InstanceData->ControlRig, SharedData->OutputMapping, InOutput.Curves);
+	Params.ControlRigVariableMappings.UpdateCurveOutputs(Params.ControlRig, Params.OutputMapping, InOutput.Curves);
 }
 
-bool FAnimNextControlRigTask::CanExecute(const UControlRig* ControlRig) const
+bool FAnimNextControlRigTask::CanExecute() const
 {
 	if (CVarControlRigDisableExecutionAnimNext->GetInt() != 0)
 	{
 		return false;
 	}
 
-	if (!InstanceData->ControlRigHierarchyMappings.CanExecute())
+	if (!Params.ControlRigHierarchyMappings.CanExecute())
 	{
 		return false;
 	}
 
-	if (ControlRig)
+	if (Params.ControlRig)
 	{
-		return ControlRig->CanExecute();
+		return Params.ControlRig->CanExecute();
 	}
 
 	return false;
 }
 
-void FAnimNextControlRigTask::QueueControlRigDrawInstructions(UControlRig* ControlRig, FRigVMDrawInterface* DebugDrawInterface, const FTransform& ComponentTransform) const
+void FAnimNextControlRigTask::QueueControlRigDrawInstructions(FRigVMDrawInterface* InDebugDrawInterface, const FTransform& InComponentTransform) const
 {
-	ensure(ControlRig);
-	ensure(DebugDrawInterface);
+	ensure(Params.ControlRig);
+	ensure(InDebugDrawInterface);
 
-	if (ControlRig && DebugDrawInterface)
+	if (Params.ControlRig && InDebugDrawInterface)
 	{
-		for (FRigVMDrawInstruction& Instruction : ControlRig->GetDrawInterface().Instructions)
+		for (FRigVMDrawInstruction& Instruction : Params.ControlRig->GetDrawInterface().Instructions)
 		{
 			if (!Instruction.IsValid())
 			{
 				continue;
 			}
 
-			const FTransform InstructionTransform = Instruction.Transform * ComponentTransform;
+			const FTransform InstructionTransform = Instruction.Transform * InComponentTransform;
 			Instruction.Transform = InstructionTransform;
-			DebugDrawInterface->DrawInstruction(Instruction);
+			InDebugDrawInterface->DrawInstruction(Instruction);
 		}
 	}
 }
+
+#if UE_ENABLE_DEBUG_DRAWING
+void FAnimNextControlRigTask::SetComponentTransform(const FTransform& InComponentTransform)
+{
+	Params.ComponentTransform = InComponentTransform;
+}
+#endif
\ No newline at end of file
diff --git a/Engine/Plugins/Experimental/UAF/UAFControlRig/Source/UAFControlRig/Private/ControlRigTask.h b/Engine/Plugins/Experimental/UAF/UAFControlRig/Source/UAFControlRig/Private/ControlRigTask.h
deleted file mode 100644
diff --git a/Engine/Plugins/Experimental/UAF/UAFControlRig/Source/UAFControlRig/Private/ControlRigTrait.cpp b/Engine/Plugins/Experimental/UAF/UAFControlRig/Source/UAFControlRig/Private/ControlRigTrait.cpp
--- a/Engine/Plugins/Experimental/UAF/UAFControlRig/Source/UAFControlRig/Private/ControlRigTrait.cpp
+++ b/Engine/Plugins/Experimental/UAF/UAFControlRig/Source/UAFControlRig/Private/ControlRigTrait.cpp
@@ -18,6 +18,7 @@
 #include "Animation/AnimNodeBase.h"
 #include "AnimationDataSource.h"
 #include "RigVMCore/RigVMMemoryStorageStruct.h"
+#include "ControlRigTask.h"
 #if WITH_EDITOR
 #include "RigVMModel/RigVMController.h"
 #include "ControlRigBlueprintLegacy.h" // to get the preview skeleton
@@ -197,6 +198,7 @@ void FControlRigTrait::FInstanceData::Construct(const FExecutionContext& Context
 	FTrait::FInstanceData::Construct(Context, Binding);
 
 	const FSharedData* SharedData = Binding.GetSharedData<FSharedData>();
+	check(SharedData);
 
 #if WITH_EDITOR
 	{
@@ -211,14 +213,26 @@ void FControlRigTrait::FInstanceData::Construct(const FExecutionContext& Context
 			USkeletalMeshComponent* BindableObject = const_cast<USkeletalMeshComponent*>(FInstanceData::GetBindableObject(Context));
 			if (ensure(CreateControlRig(GetAnimContext(Context), BindableObject, SharedData->ControlRigClass, this)))
 			{
+				FAnimNextControlRigTaskParams Params;
+				Params.ControlRig = ControlRig;
+				Params.bResetInputPoseToInitial = SharedData->bResetInputPoseToInitial;
+				Params.bTransferInputPose = SharedData->bTransferInputPose;
+				Params.bTransferInputCurves = SharedData->bTransferInputCurves;
+				Params.bSetRefPoseFromSkeleton = SharedData->bSetRefPoseFromSkeleton;
+				Params.bTransferPoseInGlobalSpace = SharedData->bTransferPoseInGlobalSpace;
+				Params.EventQueue = SharedData->EventQueue;
+				Params.InputMapping = SharedData->InputMapping;
+				Params.OutputMapping = SharedData->OutputMapping;
+#if UE_ENABLE_DEBUG_DRAWING
+				Params.DebugDrawInterface = Context.GetDebugDrawInterface();
+#endif
+
+				EvaluationTask = MakeShared<FAnimNextControlRigTask>(MoveTemp(Params));
+
 				InitializeControlRig(Context, Binding);
 			}
 		}
 	}
-
-#if UE_ENABLE_DEBUG_DRAWING
-	DebugDrawInterface = Context.GetDebugDrawInterface();
-#endif
 }
 
 void FControlRigTrait::FInstanceData::Destruct(const FExecutionContext& Context, const FTraitBinding& Binding)
@@ -229,33 +243,18 @@ void FControlRigTrait::FInstanceData::Destruct(const FExecutionContext& Context,
 	UE::UAF::ControlRig::FAnimNextControlRigModule::OnObjectsReinstanced.Remove(OnObjectsReinstancedHandle);
 #endif
 
+	EvaluationTask.Reset();
+
 	DestroyControlRig(Context, Binding);
 }
 
 void FControlRigTrait::FInstanceData::InitializeControlRig(const FExecutionContext& Context, const FTraitBinding& Binding)
 {
 	const FSharedData* SharedData = Binding.GetSharedData<FSharedData>();
-
+	check(EvaluationTask);
 	if (ControlRig)
 	{
-		// Provide available properties to the construction event
-		InitializeCustomProperties(Binding, SharedData);
-		ControlRigVariableMappings.PropagateCustomInputProperties(ControlRig);
-
-		ControlRig->Initialize(true);
-		ControlRig->RequestInit();
-
-		ControlRigHierarchyMappings.ResetRefPoseSetterHash();
-
-		ControlRigVariableMappings.ResetCurvesInputToControlCache();
-		ControlRigHierarchyMappings.ResetRefPoseSetterHash();
-		ControlRigVariableMappings.CacheCurveMappings(SharedData->InputMapping, SharedData->OutputMapping, ControlRig->GetHierarchy());
-
-		// Re-init Custom Properties after construction, as new controls could be created and might have to be remapped
-		InitializeCustomProperties(Binding, SharedData);
-		ControlRigVariableMappings.PropagateCustomInputProperties(ControlRig);
-
-		//UpdateInputOutputMappingIfRequired(ControlRig, InRequiredBones);
+		EvaluationTask->InitializeControlRig(&PropertyMappings);
 
 		ControlRig->OnInitialized_AnyThread().RemoveAll(this);
 		OnInitializedHandle = ControlRig->OnInitialized_AnyThread().AddRaw(this, &FControlRigTrait::FInstanceData::HandleOnInitialized_AnyThread);
@@ -411,7 +410,9 @@ int32 FControlRigTrait::FInstanceData::GetExposedControlsData(const UE::UAF::FTr
 
 void FControlRigTrait::FInstanceData::HandleOnInitialized_AnyThread(URigVMHost*, const FName&)
 {
-	ControlRigHierarchyMappings.ResetRefPoseSetterHash();
+	check(EvaluationTask);
+	EvaluationTask->OnControlRigInitialized();
+
 	#if WITH_EDITOR
 		bRegenerateVariableMappings = true;	// required as FRigVMEditorModule::PreChange (UUSerStructs) recreates VM memory and requests a re-init, which recreates controls
 	#endif
@@ -419,10 +420,9 @@ void FControlRigTrait::FInstanceData::HandleOnInitialized_AnyThread(URigVMHost*,
 
 void FControlRigTrait::FInstanceData::InitializeCustomProperties(const FTraitBinding& Binding, const FSharedData* SharedData)
 {
-	UE::UAF::FControlRigTrait::FInstanceData* InstanceData = Binding.GetInstanceData<UE::UAF::FControlRigTrait::FInstanceData>();
-	 
+	check(EvaluationTask);
 	// Setup mappings using the latent pin memory as source (we have to copy from latent pin to external variable / rig control)
-	ControlRigVariableMappings.InitializeCustomProperties(ControlRig, InstanceData->PropertyMappings);
+	EvaluationTask->GetControlRigVariableMappings().InitializeCustomProperties(ControlRig, PropertyMappings);
 }
 
 #if WITH_EDITOR
@@ -558,19 +558,9 @@ void FControlRigTrait::OnBecomeRelevant(FUpdateTraversalContext& Context, const
 
 	FInstanceData* InstanceData = Binding.GetInstanceData<FInstanceData>();
 	check(InstanceData);
+	check(InstanceData->EvaluationTask);
 
-	if (UControlRig* ControlRig = GetControlRig(InstanceData))
-	{
-		ControlRig->RequestInit();
-
-		InstanceData->bControlRigRequiresInitialization = true;
-		InstanceData->LastBonesSerialNumberForCacheBones = 0;
-	}
-
-	InstanceData->ControlRigHierarchyMappings.InitializeInstance();
-	InstanceData->ControlRigHierarchyMappings.ResetRefPoseSetterHash();
-
-	InstanceData->bUpdateInputOutputMapping = true;
+	InstanceData->EvaluationTask->RequestInit();
 }
 
 void FControlRigTrait::PreUpdate(FUpdateTraversalContext& Context, const TTraitBinding<IUpdate>& Binding, const FTraitUpdateState& TraitState) const
@@ -616,7 +606,7 @@ void FControlRigTrait::PreUpdate(FUpdateTraversalContext& Context, const TTraitB
 		const float DeltaTime = TraitState.GetDeltaTime();
 		ControlRig->SetDeltaTime(DeltaTime);
 
-		InstanceData->ControlRigVariableMappings.PropagateCustomInputProperties(ControlRig);
+		InstanceData->EvaluationTask->GetControlRigVariableMappings().PropagateCustomInputProperties(ControlRig);
 	}
 }
 
@@ -625,6 +615,8 @@ void FControlRigTrait::PostEvaluate(FEvaluateTraversalContext& Context, const TT
 {
 	const FSharedData* SharedData = Binding.GetSharedData<FSharedData>();
 	FInstanceData* InstanceData = Binding.GetInstanceData<FInstanceData>();
+	check(InstanceData);
+	check(InstanceData->EvaluationTask);
 
 	if (UControlRig* ControlRig = GetControlRig(InstanceData))
 	{
@@ -632,11 +624,11 @@ void FControlRigTrait::PostEvaluate(FEvaluateTraversalContext& Context, const TT
 #if UE_ENABLE_DEBUG_DRAWING
 		if (const USkeletalMeshComponent* BindableObject = FInstanceData::GetBindableObject(Context))
 		{
-			InstanceData->ComponentTransform = BindableObject->GetComponentTransform();
+			InstanceData->EvaluationTask->SetComponentTransform(BindableObject->GetComponentTransform());
 		}
 #endif
 
-		Context.AppendTask(FAnimNextControlRigTask::Make(SharedData, InstanceData));
+		Context.AppendTaskPtr(InstanceData->EvaluationTask);
 	}
 }
 
diff --git a/Engine/Plugins/Experimental/UAF/UAFControlRig/Source/UAFControlRig/Public/ControlRigTask.h b/Engine/Plugins/Experimental/UAF/UAFControlRig/Source/UAFControlRig/Public/ControlRigTask.h
new file mode 100644
--- /dev/null
+++ b/Engine/Plugins/Experimental/UAF/UAFControlRig/Source/UAFControlRig/Public/ControlRigTask.h
@@ -0,0 +1,132 @@
+// Copyright Epic Games, Inc. All Rights Reserved.
+
+#pragma once
+
+#include "CoreMinimal.h"
+#include "EvaluationVM/EvaluationTask.h"
+#include "Animation/AnimCurveTypes.h"
+#include "LODPose.h"
+#include "AnimNextStats.h"
+#include "ControlRig.h"
+#include "Delegates/IDelegateInstance.h"
+#include "AnimNextControlRigHierarchyMappings.h"
+#include "Tools/ControlRigVariableMappings.h"
+#include "Tools/ControlRigIOSettings.h"
+#if WITH_EDITOR
+#include "ControlRigIOMapping.h"
+#endif
+
+// --- ---
+#include "ControlRigTask.generated.h"
+
+namespace UE::UAF
+{
+struct FKeyframeState;
+}
+
+DECLARE_CYCLE_STAT_EXTERN(TEXT("UAF Task: ControlRig"), STAT_AnimNext_Task_ControlRig, STATGROUP_AnimNext, UAFCONTROLRIG_API);
+
+USTRUCT()
+struct FControlRigEventName
+{
+	GENERATED_BODY()
+
+	FControlRigEventName()
+		: EventName(NAME_None)
+	{
+	}
+
+	UPROPERTY(EditAnywhere, Category = Links)
+	FName EventName;
+};
+
+USTRUCT()
+struct FAnimNextControlRigTaskParams
+{
+	GENERATED_BODY()
+
+	bool bResetInputPoseToInitial = true;
+	bool bTransferInputPose = true;
+	bool bTransferInputCurves = true;
+	bool bSetRefPoseFromSkeleton = false;
+	bool bTransferPoseInGlobalSpace = false;
+	TArray<FControlRigEventName> EventQueue;
+	TMap<FName, FName> InputMapping;
+	TMap<FName, FName> OutputMapping;
+
+	FRigVMDrawInterface* DebugDrawInterface = nullptr;
+
+	FControlRigVariableMappings ControlRigVariableMappings;
+	mutable UE::UAF::ControlRig::FAnimNextControlRigHierarchyMappings ControlRigHierarchyMappings;
+
+	TObjectPtr<UControlRig> ControlRig;
+
+	uint16 LastBonesSerialNumberForCacheBones = 0;
+	mutable bool bControlRigRequiresInitialization = false;
+	mutable bool bUpdateInputOutputMapping = false;
+	FTransform ComponentTransform;
+
+	bool bConsumesPreviousPose = true;
+};
+
+USTRUCT()
+struct UAFCONTROLRIG_API FAnimNextControlRigTask : public FAnimNextEvaluationTask
+{
+	GENERATED_BODY()
+
+	DECLARE_ANIM_EVALUATION_TASK(FAnimNextControlRigTask)
+
+	FAnimNextControlRigTask() = default;
+	explicit FAnimNextControlRigTask(FAnimNextControlRigTaskParams&& InParams)
+		: Params(MoveTemp(InParams))
+	{
+	}
+
+	// Task entry point
+	virtual void Execute(UE::UAF::FEvaluationVM& VM) const override;
+
+public:
+	const FAnimNextControlRigTaskParams& GetParams() const { return Params; }
+
+	// Initializes the control rig
+	void InitializeControlRig(const FControlRigVariableMappings::FCustomPropertyMappings* PropertyMappings = nullptr);
+
+	// Marks the control rig as needing initialization on the next execute
+	void RequestInit();
+
+	// Called when the control rig has been initialized
+	void OnControlRigInitialized();
+
+#if UE_ENABLE_DEBUG_DRAWING
+	void SetComponentTransform(const FTransform& InComponentTransform);
+#endif
+
+	FControlRigVariableMappings& GetControlRigVariableMappings() { return Params.ControlRigVariableMappings; }
+	const FControlRigVariableMappings& GetControlRigVariableMappings() const { return Params.ControlRigVariableMappings; }
+
+	UE::UAF::ControlRig::FAnimNextControlRigHierarchyMappings& GetControlRigHierarchyMappings() { return Params.ControlRigHierarchyMappings; }
+	const UE::UAF::ControlRig::FAnimNextControlRigHierarchyMappings& GetControlRigHierarchyMappings() const { return Params.ControlRigHierarchyMappings; }
+
+	UControlRig* GetControlRig() const { return Params.ControlRig; }
+
+private:
+	FAnimNextControlRigTaskParams Params;
+
+	FControlRigIOSettings InputSettings;
+	FControlRigIOSettings OutputSettings;
+
+	mutable int32 LastLOD = INDEX_NONE;
+
+	bool bExecute = true;
+	mutable bool bClearEventQueueRequired = false;
+	void ExecuteControlRig(UE::UAF::FEvaluationVM& VM, UE::UAF::FKeyframeState& KeyFrameState) const;
+	void UpdateInput(UE::UAF::FEvaluationVM& VM, UE::UAF::FKeyframeState& KeyFrameState, UE::UAF::FKeyframeState& InOutput) const;
+	void UpdateOutput(UE::UAF::FEvaluationVM& VM, UE::UAF::FKeyframeState& KeyFrameState, UE::UAF::FKeyframeState& InOutput) const;
+
+	bool CanExecute() const;
+
+	void QueueControlRigDrawInstructions(FRigVMDrawInterface* DebugDrawInterface, const FTransform& ComponentTransform) const;
+
+
+
+};
diff --git a/Engine/Plugins/Experimental/UAF/UAFControlRig/Source/UAFControlRig/UAFControlRig.Build.cs b/Engine/Plugins/Experimental/UAF/UAFControlRig/Source/UAFControlRig/UAFControlRig.Build.cs
--- a/Engine/Plugins/Experimental/UAF/UAFControlRig/Source/UAFControlRig/UAFControlRig.Build.cs
+++ b/Engine/Plugins/Experimental/UAF/UAFControlRig/Source/UAFControlRig/UAFControlRig.Build.cs
@@ -21,10 +21,16 @@ namespace UnrealBuildTool.Rules
 
 			if (Target.bBuildEditor == true)
 			{
+				PublicDependencyModuleNames.AddRange(
+					new string[]
+					{
+						"ControlRigDeveloper"
+					}
+				);
+
 				PrivateDependencyModuleNames.AddRange(
 				   new string[]
 				   {
-						"ControlRigDeveloper",
 						"WorkspaceEditor",
 						"SlateCore",
 						"Slate",
