Fix Control Rig customization not finding the shared data pin, due to a prefix discreapancy

diff --git a/Engine/Plugins/Experimental/UAF/UAFControlRig/Source/UAFControlRigEditor/Private/ControlRigTraitCustomization.cpp b/Engine/Plugins/Experimental/UAF/UAFControlRig/Source/UAFControlRigEditor/Private/ControlRigTraitCustomization.cpp
--- a/Engine/Plugins/Experimental/UAF/UAFControlRig/Source/UAFControlRigEditor/Private/ControlRigTraitCustomization.cpp
+++ b/Engine/Plugins/Experimental/UAF/UAFControlRig/Source/UAFControlRigEditor/Private/ControlRigTraitCustomization.cpp
@@ -127,6 +127,8 @@ void FControlRigTraitSharedDataCustomization::OnPropertyExposeCheckboxChanged(EC
 						static FName ExposedPropertyControlDefaultValuesName = GET_MEMBER_NAME_CHECKED(FControlRigTraitSharedData, ExposedPropertyControlDefaultValues);
 						
 						const FString ControlRigTraitName = FindControlRigTraitPinName(ModelNode);
+						ensureAlwaysMsgf(!ControlRigTraitName.IsEmpty(), TEXT("Can not Find a Control Rig Trait pin by name at node : [%s]."), *ModelNode->GetFName().ToString());
+
 						const FString ControlRigTraitPrefix(ControlRigTraitName + TEXT("."));
 						const FString ExposedVariablesSubPath(ControlRigTraitPrefix + ExposedPropertyVariablesName.ToString());
 						const FString ExposedControlsSubPath(ControlRigTraitPrefix + ExposedPropertyControlsName.ToString());
@@ -406,6 +408,7 @@ FString FControlRigTraitSharedDataCustomization::FindControlRigTraitPinName(URig
 	if (const FTrait* Trait = TraitRegistry.Find(FControlRigTrait::TraitUID)) // the context menu in the traits create the traits using a sanitized display name
 	{
 		ControlRigTraitName = Trait->GetTraitName();
+		ControlRigTraitName.RemoveAt(0);
 		FString DisplayNameMetadata;
 		Trait->GetTraitSharedDataStruct()->GetStringMetaDataHierarchical(FRigVMStruct::DisplayNameMetaName, &DisplayNameMetadata);
 		ControlRigTraitDisplayName = DisplayNameMetadata.IsEmpty() ? Trait->GetTraitName() : DisplayNameMetadata;
@@ -413,7 +416,7 @@ FString FControlRigTraitSharedDataCustomization::FindControlRigTraitPinName(URig
 	}
 	else
 	{
-		ControlRigTraitName = FString(TEXT("FControlRigTrait"));
+		ControlRigTraitName = FString(TEXT("ControlRigTrait"));
 	}
 
 	FString Tmp;
