Switched check to ensure for EvaluationTask in ControlRigTrait

#Synced-CL 48363837

diff --git a/Engine/Plugins/Experimental/UAF/UAFControlRig/Source/UAFControlRig/Private/ControlRigTrait.cpp b/Engine/Plugins/Experimental/UAF/UAFControlRig/Source/UAFControlRig/Private/ControlRigTrait.cpp
--- a/Engine/Plugins/Experimental/UAF/UAFControlRig/Source/UAFControlRig/Private/ControlRigTrait.cpp
+++ b/Engine/Plugins/Experimental/UAF/UAFControlRig/Source/UAFControlRig/Private/ControlRigTrait.cpp
@@ -559,9 +559,10 @@ void FControlRigTrait::OnBecomeRelevant(FUpdateTraversalContext& Context, const
 
 	FInstanceData* InstanceData = Binding.GetInstanceData<FInstanceData>();
 	check(InstanceData);
-	check(InstanceData->EvaluationTask);
-
-	InstanceData->EvaluationTask->RequestInit();
+	if (ensure(InstanceData->EvaluationTask))
+	{
+		InstanceData->EvaluationTask->RequestInit();
+	}
 }
 
 void FControlRigTrait::PreUpdate(FUpdateTraversalContext& Context, const TTraitBinding<IUpdate>& Binding, const FTraitUpdateState& TraitState) const
@@ -617,19 +618,20 @@ void FControlRigTrait::PostEvaluate(FEvaluateTraversalContext& Context, const TT
 	const FSharedData* SharedData = Binding.GetSharedData<FSharedData>();
 	FInstanceData* InstanceData = Binding.GetInstanceData<FInstanceData>();
 	check(InstanceData);
-	check(InstanceData->EvaluationTask);
-
-	if (UControlRig* ControlRig = GetControlRig(InstanceData))
+	if (ensure(InstanceData->EvaluationTask))
 	{
-		// The transform is used by the task to "transform" the debug drawings inside Control Rig
-#if UE_ENABLE_DEBUG_DRAWING
-		if (const USkeletalMeshComponent* BindableObject = FInstanceData::GetBindableObject(Context))
+		if (UControlRig* ControlRig = GetControlRig(InstanceData))
 		{
-			InstanceData->EvaluationTask->SetComponentTransform(BindableObject->GetComponentTransform());
-		}
+			// The transform is used by the task to "transform" the debug drawings inside Control Rig
+#if UE_ENABLE_DEBUG_DRAWING
+			if (const USkeletalMeshComponent* BindableObject = FInstanceData::GetBindableObject(Context))
+			{
+				InstanceData->EvaluationTask->SetComponentTransform(BindableObject->GetComponentTransform());
+			}
 #endif
 
-		Context.AppendTaskPtr(InstanceData->EvaluationTask);
+			Context.AppendTaskPtr(InstanceData->EvaluationTask);
+		}
 	}
 }
 
