Fix server target compile

#Synced-CL 48727127

diff --git a/Engine/Plugins/Experimental/UAF/UAFControlRig/Source/UAFControlRig/Private/ControlRigTask.cpp b/Engine/Plugins/Experimental/UAF/UAFControlRig/Source/UAFControlRig/Private/ControlRigTask.cpp
--- a/Engine/Plugins/Experimental/UAF/UAFControlRig/Source/UAFControlRig/Private/ControlRigTask.cpp
+++ b/Engine/Plugins/Experimental/UAF/UAFControlRig/Source/UAFControlRig/Private/ControlRigTask.cpp
@@ -13,6 +13,7 @@
 #include "Animation/AnimNodeBase.h"
 #include "Units/Execution/RigUnit_PrepareForExecution.h"
 #include "ControlRigObjectBinding.h"
+#include "Misc/ScopeLock.h"
 #if ENABLE_ANIM_DEBUG
 #include "AnimNode_ControlRigBase.h"
 #endif
@@ -135,7 +136,7 @@ void FAnimNextControlRigTask::ExecuteControlRig(UE::UAF::FEvaluationVM& VM, UE::
 	if (Params.ControlRig)
 	{
 		// Before we start modifying the RigHierarchy, we need to lock the rig to avoid corrupted state
-		FScopeLock LockRig(&Params.ControlRig->GetEvaluateMutex());
+		UE::TScopeLock LockRig(Params.ControlRig->GetEvaluateMutex());
 
 		URigHierarchy* Hierarchy = Params.ControlRig->GetHierarchy();
 		if (Hierarchy == nullptr)
