UAF: Fix control rig pins not being exposed on control rig trait

#Synced-CL 49962543

diff --git a/Engine/Plugins/Experimental/UAF/UAFControlRig/Source/UAFControlRigEditor/Private/ControlRigTraitCustomization.cpp b/Engine/Plugins/Experimental/UAF/UAFControlRig/Source/UAFControlRigEditor/Private/ControlRigTraitCustomization.cpp
--- a/Engine/Plugins/Experimental/UAF/UAFControlRig/Source/UAFControlRigEditor/Private/ControlRigTraitCustomization.cpp
+++ b/Engine/Plugins/Experimental/UAF/UAFControlRig/Source/UAFControlRigEditor/Private/ControlRigTraitCustomization.cpp
@@ -22,6 +22,7 @@
 #include "K2Node.h"
 #include "TraitCore/TraitRegistry.h"
 #include "RigVMStringUtils.h"
+#include "Templates/UAFGraphNodeTemplate.h"
 
 #define LOCTEXT_NAMESPACE "ControlRigTraitSharedDataCustomization"
 
@@ -205,6 +206,25 @@ void FControlRigTraitSharedDataCustomization::OnPropertyExposeCheckboxChanged(EC
 						}
 
 						Controller->RepopulatePinsOnNode(ModelNode, true, false, true);
+
+						// Now that pins have been added programmatically, we must also add them to the node layout or they will not be visible
+						TArray<URigVMPin*> Pins = ModelNode->GetAllPinsRecursively();
+						FRigVMNodeLayout Layout = ModelNode->GetNodeLayout();
+						TArray<FString> PinPathsRequiringCategory;
+						PinPathsRequiringCategory.Reserve(Pins.Num());
+						for(URigVMPin* Pin : Pins)
+						{
+							if (Pin->GetIndexInCategory() == INDEX_NONE && Pin->GetDirection() != ERigVMPinDirection::Hidden)
+							{
+								FString PinPath = Pin->GetPinPath(false); // this contains the node name as prefix, but the node layout doesn't want that
+								FString NodePrefix = ModelNode ->GetName() + ".";
+								PinPath.RemoveFromStart(NodePrefix);
+								PinPathsRequiringCategory.Add(PinPath);
+							}
+						}
+
+						UUAFGraphNodeTemplate::SetCategoryForPinsInLayout(PinPathsRequiringCategory, TEXT("Control Rig"), Layout, true);
+						Controller->SetNodeLayout(ModelNode->GetFName(), Layout);
 					}
 				}
 			}
