UAF: Move SetBindableObject within scope of GCScopeGuard in FControlRigTrait::CreateControlRig to avoid GC ensure when running in Debug

#Synced-CL 50345329

diff --git a/Engine/Plugins/Experimental/UAF/UAFControlRig/Source/UAFControlRig/Private/ControlRigTrait.cpp b/Engine/Plugins/Experimental/UAF/UAFControlRig/Source/UAFControlRig/Private/ControlRigTrait.cpp
--- a/Engine/Plugins/Experimental/UAF/UAFControlRig/Source/UAFControlRig/Private/ControlRigTrait.cpp
+++ b/Engine/Plugins/Experimental/UAF/UAFControlRig/Source/UAFControlRig/Private/ControlRigTrait.cpp
@@ -674,14 +674,17 @@ bool FControlRigTrait::CreateControlRig(UObject* InAnimContext, UObject* InBinda
 {
 	if (InstanceData->ControlRig == nullptr && InControlRigClass != nullptr)
 	{
-		// Let's make sure the GC isn't running when we try to create a new Control Rig.
+		bool bSuccess = false;
+		// Let's make sure the GC isn't running when we try to create a new Control Rig 
+		// or when setting the bindable object as GetDataSourceRegistry might create a new object
 		{
 			FGCScopeGuard GCGuard;
 			InstanceData->ControlRig = NewObject<UControlRig>(InAnimContext, InControlRigClass);
 			InstanceData->ControlRig->AtomicallyClearInternalFlags(EInternalObjectFlags::Async);
+			
+			bSuccess = SetBindableObject(InstanceData->ControlRig, InAnimContext, InBindableObject);
 		}
 
-		const bool bSuccess = SetBindableObject(InstanceData->ControlRig, InAnimContext, InBindableObject);
 #if WITH_EDITOR
 		InstanceData->bRefreshBindableObject = !bSuccess;
 #endif
