Control Rig: Fix ControlRigTrait not finding newer BP-independent rigs
#rnx

#Synced-CL 51375008

diff --git a/Engine/Plugins/Animation/ControlRig/Source/ControlRigEditor/Private/ControlRigAssetReferenceDetails.cpp b/Engine/Plugins/Animation/ControlRig/Source/ControlRigEditor/Private/ControlRigAssetReferenceDetails.cpp
--- a/Engine/Plugins/Animation/ControlRig/Source/ControlRigEditor/Private/ControlRigAssetReferenceDetails.cpp
+++ b/Engine/Plugins/Animation/ControlRig/Source/ControlRigEditor/Private/ControlRigAssetReferenceDetails.cpp
@@ -120,12 +120,14 @@ TSharedRef<SWidget> FControlRigAssetReferenceDetails::MakeComboMenuContent()
 	return SNew(SControlRigAssetReferencePicker)
 		.OnSelectionChanged_Lambda([this](FControlRigAssetSoftReference InReference)
 			{
-				if (WeakStructPropertyHandle.IsValid())
+				if (TSharedPtr<IPropertyHandle> PropertyHandle = WeakStructPropertyHandle.Pin())
 				{
 					void* StructData = nullptr;
-					const FPropertyAccess::Result Result = WeakStructPropertyHandle.Pin()->GetValueData(StructData);
+					const FPropertyAccess::Result Result = PropertyHandle->GetValueData(StructData);
 					if(Result == FPropertyAccess::Success)
 					{
+						PropertyHandle->NotifyPreChange();
+						
 						check(StructData);
 						if (bPropertyIsStrongReference)
 						{
@@ -137,36 +139,16 @@ TSharedRef<SWidget> FControlRigAssetReferenceDetails::MakeComboMenuContent()
 							FControlRigAssetSoftReference* Data = static_cast<FControlRigAssetSoftReference*>(StructData);
 							*Data = InReference;
 						}
-						if (ComboButton.IsValid())
-						{
-							ComboButton->SetIsOpen(false);
-						}
+						
+						SelectedSource = InReference;
+						PropertyHandle->NotifyPostChange(EPropertyChangeType::ValueSet);
+						PropertyHandle->NotifyFinishedChangingProperties();
 					}
-				}
-				SelectedSource = InReference;
-				if (TSharedPtr<IPropertyHandle> PropertyHandle = WeakStructPropertyHandle.Pin())
-				{
-					PropertyHandle->NotifyPreChange();
 					
-					void* StructData = nullptr;
-					const FPropertyAccess::Result Result = PropertyHandle->GetValueData(StructData);
-					if(Result == FPropertyAccess::Success)
+					if (ComboButton.IsValid())
 					{
-						check(StructData);
-						if (bPropertyIsStrongReference)
-						{
-							FControlRigAssetStrongReference* Data = static_cast<FControlRigAssetStrongReference*>(StructData);
-							*Data = InReference.LoadStrongReference();
-						}
-						else
-						{
-							FControlRigAssetSoftReference* Data = static_cast<FControlRigAssetSoftReference*>(StructData);
-							*Data = InReference;
-						}
+						ComboButton->SetIsOpen(false);
 					}
-
-					PropertyHandle->NotifyPostChange(EPropertyChangeType::ValueSet);
-					PropertyHandle->NotifyFinishedChangingProperties();
 				}
 			});
 }
diff --git a/Engine/Plugins/Experimental/UAF/UAF/Source/UAF/Public/Variables/UAFInstanceVariableData.h b/Engine/Plugins/Experimental/UAF/UAF/Source/UAF/Public/Variables/UAFInstanceVariableData.h
--- a/Engine/Plugins/Experimental/UAF/UAF/Source/UAF/Public/Variables/UAFInstanceVariableData.h
+++ b/Engine/Plugins/Experimental/UAF/UAF/Source/UAF/Public/Variables/UAFInstanceVariableData.h
@@ -24,7 +24,7 @@ class IRigVMRuntimeAssetInterface;
 
 namespace UE::UAF
 {
-struct FControlRigTrait;
+template<typename TSharedData> struct TControlRigTraitBase;
 struct FVariableOverrides;
 	struct FVariableOverridesCollection;
 	struct FUAFTrace;
@@ -50,7 +50,7 @@ private:
 	friend UE::UAF::FUAFTrace;
 	friend UE::UAF::FInstanceTaskContext;
 	friend FUAFRigVMComponent;
-	friend UE::UAF::FControlRigTrait;
+	template<typename T> friend struct UE::UAF::TControlRigTraitBase;
 
 	// Initialize from a host's asset and its references
 	void Initialize(const FUAFAssetInstance& InHostInstance, const TSharedPtr<const UE::UAF::FVariableOverridesCollection>& InOverrides);
diff --git a/Engine/Plugins/Experimental/UAF/UAFControlRig/Source/UAFControlRig/Internal/ControlRigTrait.h b/Engine/Plugins/Experimental/UAF/UAFControlRig/Source/UAFControlRig/Internal/ControlRigTrait.h
--- a/Engine/Plugins/Experimental/UAF/UAFControlRig/Source/UAFControlRig/Internal/ControlRigTrait.h
+++ b/Engine/Plugins/Experimental/UAF/UAFControlRig/Source/UAFControlRig/Internal/ControlRigTrait.h
@@ -43,18 +43,15 @@ struct FControlRigExposedProperty
 	bool bIsVaraible = false;
 };
 
-USTRUCT(meta = (DisplayName = "Control Rig"))
-struct FControlRigTraitSharedData : public FAnimNextTraitSharedData
+USTRUCT()
+struct FControlRigTraitSharedDataBase : public FAnimNextTraitSharedData
 {
 	GENERATED_BODY()
-
+	
 	/** Input to be processed. */
 	UPROPERTY()
 	FAnimNextTraitHandle Input;
 
-	UPROPERTY(EditAnywhere, Category = ControlRig, meta = (Hidden, DisallowCreateNew))
-	TSubclassOf<UControlRig> ControlRigClass;
-
 	// Skeleton to use as source to extract the Rig Controls. If null, the system will use the preview skeleton that was used to create the rig.
 	UPROPERTY(EditAnywhere, Category = ControlRig, meta = (Hidden))
 	TObjectPtr<USkeleton> ControlRigSkeleton;
@@ -138,22 +135,65 @@ struct FControlRigTraitSharedData : public FAnimNextTraitSharedData
 	TArray<FString> ExposedPropertyControlDefaultValues;
 
 	// This is the array of latent input properties
-	// This is computed at load time based on the selected ControlRig class
+	// This is computed at load time based on the selected ControlRig asset
 	UPROPERTY(Transient, meta = (Hidden))
 	TArray<FName> LatentProperties;
 
 	// This is the array that maps a latent property to its memory layout in the trait instance data
-	// This is computed at load time based on the selected ControlRig class
+	// This is computed at load time based on the selected ControlRig asset
 	UPROPERTY(Transient, meta = (Hidden))
 	TArray<uint32> LatentPropertyMemoryLayouts;
 
+	// Manual handling of latent properties
+	static void InitializeLatentProperties(const TArray<FControlRigVariableMappings::FCustomPropertyData>& Mappings);
+
+	static void DestructLatentPropertiesValues(const TArray<FControlRigVariableMappings::FCustomPropertyData>& Mappings);
+
+	static void InitializeControlLatentPinDefaultValue(ERigControlType InControlType, uint8* InTargetLatentPinMemory);
+	
+};
+
+USTRUCT(meta = (Hidden, DisplayName = "Control Rig (Legacy)"))
+struct FControlRigTraitSharedData : public FControlRigTraitSharedDataBase
+{
+	GENERATED_BODY()
+	
+	UPROPERTY(EditAnywhere, Category = ControlRig, meta = (Hidden, DisallowCreateNew))
+	TSubclassOf<UControlRig> ControlRigClass;
+	
+	FControlRigAssetStrongReference GetControlRigAssetReference() const { return FControlRigAssetStrongReference(ControlRigClass); }
+	bool HasValidControlRigReference() const { return ControlRigClass != nullptr; }
+	TScriptInterface<const IControlRigEditorAssetInterface> GetControlRigEditorAsset() const;
+	
 	// Manual handling of latent properties
 	static void ConstructLatentProperties(const UE::UAF::FTraitBinding& Binding);
 
 	static void DestructLatentProperties(const UE::UAF::FTraitBinding& Binding);
 
-	static void InitializeControlLatentPinDefaultValue(ERigControlType InControlType, uint8* InTargetLatentPinMemory);
+	
+#if WITH_EDITOR
+	UAFCONTROLRIG_API USkeleton* GetPreviewSkeleton() const;
+#endif // WITH_EDITOR
+};
 
+USTRUCT(meta = (DisplayName = "Control Rig"))
+struct FControlRigTraitSharedData_v2 : public FControlRigTraitSharedDataBase
+{
+	GENERATED_BODY()
+	
+	UPROPERTY(EditAnywhere, Category = ControlRig, meta = (Hidden, DisallowCreateNew))
+	FControlRigAssetStrongReference ControlRigAssetReference;
+	
+	FControlRigAssetStrongReference GetControlRigAssetReference() const { return ControlRigAssetReference; }
+	bool HasValidControlRigReference() const { return ControlRigAssetReference.IsValid(); }
+	TScriptInterface<const IControlRigEditorAssetInterface> GetControlRigEditorAsset() const;
+	
+	// Manual handling of latent properties
+	static void ConstructLatentProperties(const UE::UAF::FTraitBinding& Binding);
+
+	static void DestructLatentProperties(const UE::UAF::FTraitBinding& Binding);
+
+	
 #if WITH_EDITOR
 	UAFCONTROLRIG_API USkeleton* GetPreviewSkeleton() const;
 #endif // WITH_EDITOR
@@ -162,17 +202,13 @@ struct FControlRigTraitSharedData : public FAnimNextTraitSharedData
 namespace UE::UAF
 {
 
-struct FControlRigTrait : FBaseTrait, IEvaluate, IUpdate, IHierarchy, IGarbageCollection
+template<typename TSharedData>
+struct TControlRigTraitBase : FBaseTrait, IEvaluate, IUpdate, IHierarchy, IGarbageCollection
 {
 	typedef FBaseTrait Super;
 
-	ANIM_NEXT_IMPL_DECLARE_ANIM_TRAIT_BASIC(FControlRigTrait, FBaseTrait)
-	ANIM_NEXT_IMPL_DECLARE_ANIM_TRAIT_INSTANCING_SUPPORT()
-	ANIM_NEXT_IMPL_DECLARE_ANIM_TRAIT_INTERFACE_SUPPORT()
-	ANIM_NEXT_IMPL_DECLARE_ANIM_TRAIT_EVENT_SUPPORT()
-
-	using FSharedData = FControlRigTraitSharedData;
-		
+	using FSharedData = TSharedData;
+	
 	struct FInstanceData : FTrait::FInstanceData
 	{
 		FTraitPtr Input;
@@ -194,16 +230,16 @@ struct FControlRigTrait : FBaseTrait, IEvaluate, IUpdate, IHierarchy, IGarbageCo
 
 		static UObject* GetAnimContext(const FExecutionContext& Context);
 		static const USkeletalMeshComponent* GetBindableObject(const FExecutionContext& Context);
-		static UClass* GetTargetClass(const FSharedData* SharedData);
+		static FControlRigAssetStrongReference GetTargetAssetReference(const TSharedData* SharedData);
 
 		// Returns the property, variable name and the memory of the variable (as target) and the memory of the latent property (as source)
-		static int32 GetExposedVariablesData(const UE::UAF::FTraitBinding& Binding, const FSharedData* SharedData, FControlRigVariableMappings::FCustomPropertyMappings& OutMappings);
-		static int32 GetExposedControlsData(const UE::UAF::FTraitBinding& Binding, const FSharedData* SharedData, FControlRigVariableMappings::FCustomPropertyMappings& OutMappings);
-		static void MapUAFVariablesToControlRig(const UE::UAF::FTraitBinding& Binding, const FSharedData* SharedData, FControlRigVariableMappings::FCustomPropertyMappings& OutMappings);
+		static int32 GetExposedVariablesData(const UE::UAF::FTraitBinding& Binding, const TSharedData* SharedData, FControlRigVariableMappings::FCustomPropertyMappings& OutMappings);
+		static int32 GetExposedControlsData(const UE::UAF::FTraitBinding& Binding, const TSharedData* SharedData, FControlRigVariableMappings::FCustomPropertyMappings& OutMappings);
+		static void MapUAFVariablesToControlRig(const UE::UAF::FTraitBinding& Binding, const TSharedData* SharedData, FControlRigVariableMappings::FCustomPropertyMappings& OutMappings);
 
 		void HandleOnInitialized_AnyThread(URigVMHost*, const FName&);
 
-		void InitializeCustomProperties(const FTraitBinding& Binding, const FSharedData* SharedData);
+		void InitializeCustomProperties(const FTraitBinding& Binding, const TSharedData* SharedData);
 
 #if WITH_EDITOR
 		void OnObjectsReinstanced(const TMap<UObject*, UObject*>& OldToNewInstanceMap);
@@ -235,10 +271,10 @@ struct FControlRigTrait : FBaseTrait, IEvaluate, IUpdate, IHierarchy, IGarbageCo
 protected:
 	static UControlRig* GetControlRig(FInstanceData* InstanceData);
 
-	static bool CreateControlRig(UObject* InAnimContext, UObject* InBindableObject, UClass* InControlRigClass, FInstanceData* InstanceData);
+	static bool CreateControlRig(UObject* InAnimContext, UObject* InBindableObject, FControlRigAssetStrongReference InControlRigAssetReference, FInstanceData* InstanceData);
 	static bool SetBindableObject(UControlRig* ControlRig, UObject* InAnimContext, UObject* InBindableObject);
 
-	static UClass* GetTargetClass(const FSharedData* SharedData, const TTraitBinding<IUpdate>& Binding);
+	static FControlRigAssetStrongReference GetTargetAssetReference(const FSharedData* SharedData, const TTraitBinding<IUpdate>& Binding);
 
 #if WITH_EDITOR
 	virtual void GetProgrammaticPins(FAnimNextTraitSharedData* InSharedData,
@@ -262,5 +298,25 @@ protected:
 	static bool GetVariableSizeAndAlignment(const FRigVMExternalVariable& Variable, uint32& PropertySize, uint32& PropertyAlignment);
 	static bool GetControlSizeAndAlignment(ERigControlType ControlType, uint32& PropertySize, uint32& PropertyAlignment);
 };
+	
+struct FControlRigTrait : TControlRigTraitBase<FControlRigTraitSharedData>
+{
+	using Super = TControlRigTraitBase<FControlRigTraitSharedData>;
+	
+	ANIM_NEXT_IMPL_DECLARE_ANIM_TRAIT_BASIC(FControlRigTrait, FBaseTrait)
+	ANIM_NEXT_IMPL_DECLARE_ANIM_TRAIT_INSTANCING_SUPPORT()
+	ANIM_NEXT_IMPL_DECLARE_ANIM_TRAIT_INTERFACE_SUPPORT()
+	ANIM_NEXT_IMPL_DECLARE_ANIM_TRAIT_EVENT_SUPPORT()
+};
+	
+struct FControlRigTrait_v2 : TControlRigTraitBase<FControlRigTraitSharedData_v2>
+{
+	using Super = TControlRigTraitBase<FControlRigTraitSharedData_v2>;
+
+	ANIM_NEXT_IMPL_DECLARE_ANIM_TRAIT_BASIC(FControlRigTrait_v2, FBaseTrait)
+	ANIM_NEXT_IMPL_DECLARE_ANIM_TRAIT_INSTANCING_SUPPORT()
+	ANIM_NEXT_IMPL_DECLARE_ANIM_TRAIT_INTERFACE_SUPPORT()
+	ANIM_NEXT_IMPL_DECLARE_ANIM_TRAIT_EVENT_SUPPORT()
+};
 
 } // namespace UE::UAF
diff --git a/Engine/Plugins/Experimental/UAF/UAFControlRig/Source/UAFControlRig/Private/AnimNextControlRigModule.cpp b/Engine/Plugins/Experimental/UAF/UAFControlRig/Source/UAFControlRig/Private/AnimNextControlRigModule.cpp
--- a/Engine/Plugins/Experimental/UAF/UAFControlRig/Source/UAFControlRig/Private/AnimNextControlRigModule.cpp
+++ b/Engine/Plugins/Experimental/UAF/UAFControlRig/Source/UAFControlRig/Private/AnimNextControlRigModule.cpp
@@ -25,7 +25,7 @@
 #include "UAFControlRigAssetData.h"
 #include "UAF/UAFAssetFactory.h"
 
-struct FControlRigTraitSharedData;
+struct FControlRigTraitSharedData_v2;
 
 namespace UE::UAF::ControlRig
 {
@@ -48,6 +48,21 @@ void FAnimNextControlRigModule::StartupModule()
 
 	VariableProvider = MakeShared<FAnimNextControlRigVariableProvider>();
 #endif
+	
+	FRigVMRegistry& RigVMRegistry = FRigVMRegistry::Get();
+	
+	static UScriptStruct* const AllowedStructTypes[] =
+	{
+		FControlRigAssetStrongReference::StaticStruct()
+	};
+	
+	static TPair<UClass*, FRigVMRegistry::ERegisterObjectOperation> const AllowedObjectTypes[] =
+	{
+	{ UControlRigRuntimeAsset::StaticClass(), FRigVMRegistry::ERegisterObjectOperation::Class },
+	};
+
+	RigVMRegistry.RegisterStructTypes(AllowedStructTypes);
+	RigVMRegistry.RegisterObjectTypes(AllowedObjectTypes);
 
 	RegisterGraphFactories();
 	RegisterSystemFactories();
@@ -67,22 +82,22 @@ void FAnimNextControlRigModule::ShutdownModule()
 void FAnimNextControlRigModule::RegisterGraphFactories() const
 {
 	FAnimNextFactoryParams Params;
-	Params.AddTraitStruct<FControlRigTraitSharedData>(ETraitVariableMapping::All, 0);
+	Params.AddTraitStruct<FControlRigTraitSharedData_v2>(ETraitVariableMapping::All, 0);
 
 	FAnimGraphFactory::RegisterAsset<FUAFGraphFactoryAsset_ControlRig>(
 		MoveTemp(Params),
 		[](const FUAFGraphFactoryAsset_ControlRig& AssetData, FAnimNextFactoryParams& InOutParams)
 			{
 				ensure(
-					InOutParams.AccessTraitStruct<FControlRigTraitSharedData>(0, [&AssetData, &InOutParams](FControlRigTraitSharedData& InStruct)
+					InOutParams.AccessTraitStruct<FControlRigTraitSharedData_v2>(0, [&AssetData, &InOutParams](FControlRigTraitSharedData_v2& InStruct)
 						{
-						InStruct.ControlRigClass = AssetData.ControlRigClass;
+						InStruct.ControlRigAssetReference = AssetData.ControlRigAssetReference;
 						}));
 
 				// Add control rig variables to public api for the graph
-				if (AssetData.ControlRigClass != nullptr)
+				if (AssetData.ControlRigAssetReference.IsValid())
 				{
-					InOutParams.AddPublicVariablesRigVMAsset(AssetData.ControlRigClass);
+					InOutParams.AddPublicVariablesRigVMAsset(AssetData.ControlRigAssetReference.Get());
 				}
 			});
 }
@@ -92,7 +107,7 @@ void FAnimNextControlRigModule::RegisterSystemFactories() const
 	{
 		FUAFSystemFactoryParams Params;
 		Params.AddComponent<FUAFSimplePrePhysicsGraphComponent>();
-		Params.AddPublicVariablesStruct<FControlRigTraitSharedData>(); // Q: do I need this? or is it handled by the graph factory?
+		Params.AddPublicVariablesStruct<FControlRigTraitSharedData_v2>(); // Q: do I need this? or is it handled by the graph factory?
 		Params.AddPublicVariablesStruct<FInjectionSiteTraitData>();
 
 		// Register for cooked asset
@@ -102,17 +117,17 @@ void FAnimNextControlRigModule::RegisterSystemFactories() const
 				{
 					ensure(Params.AccessVariablesStruct<FInjectionSiteTraitData>([&AssetData, &Params](FInjectionSiteTraitData& InStruct)
 						{
-							InStruct.Graph.Asset = AssetData.ControlRigClass;
+							InStruct.Graph.Asset = AssetData.ControlRigAssetReference.Get();
 						}));
 
-					ensure(Params.AccessVariablesStruct<FControlRigTraitSharedData>([&AssetData, &Params](FControlRigTraitSharedData& InStruct)
+					ensure(Params.AccessVariablesStruct<FControlRigTraitSharedData_v2>([&AssetData, &Params](FControlRigTraitSharedData_v2& InStruct)
 						{
-							InStruct.ControlRigClass = AssetData.ControlRigClass;
+							InStruct.ControlRigAssetReference = AssetData.ControlRigAssetReference;
 						}));
 
-					if (AssetData.ControlRigClass != nullptr)
+					if (AssetData.ControlRigAssetReference.IsValid())
 					{
-						Params.AddPublicVariablesRigVMAsset(AssetData.ControlRigClass);
+						Params.AddPublicVariablesRigVMAsset(AssetData.ControlRigAssetReference.Get());
 					}
 				});
 	}
diff --git a/Engine/Plugins/Experimental/UAF/UAFControlRig/Source/UAFControlRig/Private/ControlRigTrait.cpp b/Engine/Plugins/Experimental/UAF/UAFControlRig/Source/UAFControlRig/Private/ControlRigTrait.cpp
--- a/Engine/Plugins/Experimental/UAF/UAFControlRig/Source/UAFControlRig/Private/ControlRigTrait.cpp
+++ b/Engine/Plugins/Experimental/UAF/UAFControlRig/Source/UAFControlRig/Private/ControlRigTrait.cpp
@@ -28,10 +28,12 @@
 //----------------------
 // --- FSharedData ---
 //----------------------
+
+
 void FControlRigTraitSharedData::ConstructLatentProperties(const UE::UAF::FTraitBinding& Binding)
 {
 	const FControlRigTraitSharedData* SharedData = Binding.GetSharedData<FControlRigTraitSharedData>();
-	if (SharedData && SharedData->ControlRigClass.Get() != nullptr)
+	if (SharedData && SharedData->HasValidControlRigReference())
 	{
 		UE::UAF::FControlRigTrait::FInstanceData* InstanceData = Binding.GetInstanceData<UE::UAF::FControlRigTrait::FInstanceData>();
 
@@ -40,24 +42,45 @@ void FControlRigTraitSharedData::ConstructLatentProperties(const UE::UAF::FTrait
 		UE::UAF::FControlRigTrait::FInstanceData::MapUAFVariablesToControlRig(Binding, SharedData, InstanceData->PropertyMappings);
 
 		const TArray<FControlRigVariableMappings::FCustomPropertyData>& Mappings = InstanceData->PropertyMappings.GetMappings();
-		const int32 NumProperties = Mappings.Num();
+		InitializeLatentProperties(Mappings);
+	}
+}
+
+void FControlRigTraitSharedData_v2::ConstructLatentProperties(const UE::UAF::FTraitBinding& Binding)
+{
+	const FControlRigTraitSharedData_v2* SharedData = Binding.GetSharedData<FControlRigTraitSharedData_v2>();
+	if (SharedData && SharedData->HasValidControlRigReference())
+	{
+		UE::UAF::FControlRigTrait_v2::FInstanceData* InstanceData = Binding.GetInstanceData<UE::UAF::FControlRigTrait_v2::FInstanceData>();
+
+		UE::UAF::FControlRigTrait_v2::FInstanceData::GetExposedVariablesData(Binding, SharedData, InstanceData->PropertyMappings);
+		UE::UAF::FControlRigTrait_v2::FInstanceData::GetExposedControlsData(Binding, SharedData, InstanceData->PropertyMappings);
+		UE::UAF::FControlRigTrait_v2::FInstanceData::MapUAFVariablesToControlRig(Binding, SharedData, InstanceData->PropertyMappings);
+
+		const TArray<FControlRigVariableMappings::FCustomPropertyData>& Mappings = InstanceData->PropertyMappings.GetMappings();
+		InitializeLatentProperties(Mappings);
+	}
+}
+
+void FControlRigTraitSharedDataBase::InitializeLatentProperties(const TArray<FControlRigVariableMappings::FCustomPropertyData>& Mappings)
+{
+	const int32 NumProperties = Mappings.Num();
 
-		for (const FControlRigVariableMappings::FCustomPropertyData& Mapping : Mappings)
+	for (const FControlRigVariableMappings::FCustomPropertyData& Mapping : Mappings)
+	{
+		const FProperty* Property = Mapping.Property;
+		if (const uint8* LatentPinMemory = Mapping.SourceMemory)
 		{
-			const FProperty* Property = Mapping.Property;
-			if (const uint8* LatentPinMemory = Mapping.SourceMemory)
+			if (Property)
 			{
-				if (Property)
-				{
-					// Init latent memory to default value
-					uint8* MutableMemory = const_cast<uint8*>(LatentPinMemory);
-					Property->InitializeValue(MutableMemory);
-				}
-				else if (Mapping.Type == FControlRigVariableMappings::FCustomPropertyData::ECustomPropertyType::Control)
-				{
-					uint8* MutableMemory = const_cast<uint8*>(LatentPinMemory);
-					InitializeControlLatentPinDefaultValue(Mapping.ControlType, MutableMemory);
-				}
+				// Init latent memory to default value
+				uint8* MutableMemory = const_cast<uint8*>(LatentPinMemory);
+				Property->InitializeValue(MutableMemory);
+			}
+			else if (Mapping.Type == FControlRigVariableMappings::FCustomPropertyData::ECustomPropertyType::Control)
+			{
+				uint8* MutableMemory = const_cast<uint8*>(LatentPinMemory);
+				InitializeControlLatentPinDefaultValue(Mapping.ControlType, MutableMemory);
 			}
 		}
 	}
@@ -66,22 +89,39 @@ void FControlRigTraitSharedData::ConstructLatentProperties(const UE::UAF::FTrait
 void FControlRigTraitSharedData::DestructLatentProperties(const UE::UAF::FTraitBinding& Binding)
 {
 	const FControlRigTraitSharedData* SharedData = Binding.GetSharedData<FControlRigTraitSharedData>();
-	if (SharedData && SharedData->ControlRigClass.Get() != nullptr)
+	if (SharedData && SharedData->HasValidControlRigReference())
 	{
 		UE::UAF::FControlRigTrait::FInstanceData* InstanceData = Binding.GetInstanceData<UE::UAF::FControlRigTrait::FInstanceData>();
 
 		const TArray<FControlRigVariableMappings::FCustomPropertyData>& Mappings = InstanceData->PropertyMappings.GetMappings();
-		const int32 NumProperties = Mappings.Num();
+		DestructLatentPropertiesValues(Mappings);
+	}
+}
+
+void FControlRigTraitSharedData_v2::DestructLatentProperties(const UE::UAF::FTraitBinding& Binding)
+{
+	const FControlRigTraitSharedData_v2* SharedData = Binding.GetSharedData<FControlRigTraitSharedData_v2>();
+	if (SharedData && SharedData->HasValidControlRigReference())
+	{
+		UE::UAF::FControlRigTrait_v2::FInstanceData* InstanceData = Binding.GetInstanceData<UE::UAF::FControlRigTrait_v2::FInstanceData>();
+
+		const TArray<FControlRigVariableMappings::FCustomPropertyData>& Mappings = InstanceData->PropertyMappings.GetMappings();
+		DestructLatentPropertiesValues(Mappings);
+	}
+}
 
-		for (const FControlRigVariableMappings::FCustomPropertyData& Mapping : Mappings)
+void FControlRigTraitSharedDataBase::DestructLatentPropertiesValues(const TArray<FControlRigVariableMappings::FCustomPropertyData>& Mappings)
+{
+	const int32 NumProperties = Mappings.Num();
+
+	for (const FControlRigVariableMappings::FCustomPropertyData& Mapping : Mappings)
+	{
+		const FProperty* Property = Mapping.Property;
+		const uint8* LatentPinMemory = Mapping.SourceMemory;
+		if (Property && LatentPinMemory)
 		{
-			const FProperty* Property = Mapping.Property;
-			const uint8* LatentPinMemory = Mapping.SourceMemory;
-			if (Property && LatentPinMemory)
-			{
-				uint8* MutableMemory = const_cast<uint8*>(LatentPinMemory);
-				Property->DestroyValue(MutableMemory);
-			}
+			uint8* MutableMemory = const_cast<uint8*>(LatentPinMemory);
+			Property->DestroyValue(MutableMemory);
 		}
 	}
 }
@@ -89,7 +129,7 @@ void FControlRigTraitSharedData::DestructLatentProperties(const UE::UAF::FTraitB
 // Unlike mapped variables, controls does not come with an associated property, 
 // so I default init the latent property memory until the graph provides a valid value
 // This avoids setting control initial values with random memory
-void FControlRigTraitSharedData::InitializeControlLatentPinDefaultValue(ERigControlType InControlType, uint8* InTargetLatentPinMemory)
+void FControlRigTraitSharedDataBase::InitializeControlLatentPinDefaultValue(ERigControlType InControlType, uint8* InTargetLatentPinMemory)
 {
 	check(InTargetLatentPinMemory != nullptr);
 
@@ -153,12 +193,33 @@ USkeleton* FControlRigTraitSharedData::GetPreviewSkeleton() const
 {
 	USkeleton* PreviewSkeleton = ControlRigSkeleton;
 
-	if (PreviewSkeleton == nullptr && ControlRigClass != nullptr)
+	if (PreviewSkeleton == nullptr && HasValidControlRigReference())
+	{
+		// If the user has not provided an explicit skeleton to use, 
+		// as AnimNext does not have a preview skeleton, I get the one that was used to generate the rig
+		// (note that this might to be valid for some constructions and the user might have to provide the skeleton)
+		if (TScriptInterface<const IControlRigEditorAssetInterface> RigVMBlueprint = GetControlRigEditorAsset())
+		{
+			if (USkeletalMesh* SkeletalMesh = RigVMBlueprint->GetPreviewMesh())
+			{
+				PreviewSkeleton = SkeletalMesh->GetSkeleton();
+			}
+		}
+	}
+
+	return PreviewSkeleton;
+}
+
+USkeleton* FControlRigTraitSharedData_v2::GetPreviewSkeleton() const
+{
+	USkeleton* PreviewSkeleton = ControlRigSkeleton;
+
+	if (PreviewSkeleton == nullptr && HasValidControlRigReference())
 	{
 		// If the user has not provided an explicit skeleton to use, 
 		// as AnimNext does not have a preview skeleton, I get the one that was used to generate the rig
 		// (note that this might to be valid for some constructions and the user might have to provide the skeleton)
-		if (const UControlRigBlueprint* RigVMBlueprint = Cast<UControlRigBlueprint>(ControlRigClass->ClassGeneratedBy))
+		if (TScriptInterface<const IControlRigEditorAssetInterface> RigVMBlueprint = GetControlRigEditorAsset())
 		{
 			if (USkeletalMesh* SkeletalMesh = RigVMBlueprint->GetPreviewMesh())
 			{
@@ -169,12 +230,32 @@ USkeleton* FControlRigTraitSharedData::GetPreviewSkeleton() const
 
 	return PreviewSkeleton;
 }
+
+TScriptInterface<const IControlRigEditorAssetInterface> FControlRigTraitSharedData::GetControlRigEditorAsset() const
+{
+	if (ControlRigClass)
+	{
+		return ControlRigClass->ClassGeneratedBy;
+	}
+	return nullptr;
+}
+
+TScriptInterface<const IControlRigEditorAssetInterface> FControlRigTraitSharedData_v2::GetControlRigEditorAsset() const
+{
+	if (ControlRigAssetReference.IsValid())
+	{
+		return ControlRigAssetReference.GetEditorAsset();
+	}
+	return nullptr;
+}
+
 #endif //WITH_EDITOR
 
 namespace UE::UAF
 {
 
 AUTO_REGISTER_ANIM_TRAIT(FControlRigTrait)
+AUTO_REGISTER_ANIM_TRAIT(FControlRigTrait_v2)
 
 // Trait implementation boilerplate
 #define TRAIT_INTERFACE_ENUMERATOR(GeneratorMacro) \
@@ -189,13 +270,21 @@ ANIM_NEXT_IMPL_DEFINE_ANIM_TRAIT_GET_INTERFACES(FControlRigTrait, TRAIT_INTERFAC
 ANIM_NEXT_IMPL_DEFINE_ANIM_TRAIT_GET_REQUIRED_INTERFACES(FControlRigTrait, NULL_ANIM_TRAIT_INTERFACE_ENUMERATOR)
 ANIM_NEXT_IMPL_DEFINE_ANIM_TRAIT_ON_TRAIT_EVENT(FControlRigTrait, NULL_ANIM_TRAIT_EVENT_ENUMERATOR)
 ANIM_NEXT_IMPL_DEFINE_ANIM_TRAIT_GET_TRAIT_EVENTS(FControlRigTrait, NULL_ANIM_TRAIT_EVENT_ENUMERATOR)
+	
+ANIM_NEXT_IMPL_DEFINE_ANIM_TRAIT(FControlRigTrait_v2)
+ANIM_NEXT_IMPL_DEFINE_ANIM_TRAIT_GET_INTERFACE(FControlRigTrait_v2, TRAIT_INTERFACE_ENUMERATOR)
+ANIM_NEXT_IMPL_DEFINE_ANIM_TRAIT_GET_INTERFACES(FControlRigTrait_v2, TRAIT_INTERFACE_ENUMERATOR)
+ANIM_NEXT_IMPL_DEFINE_ANIM_TRAIT_GET_REQUIRED_INTERFACES(FControlRigTrait_v2, NULL_ANIM_TRAIT_INTERFACE_ENUMERATOR)
+ANIM_NEXT_IMPL_DEFINE_ANIM_TRAIT_ON_TRAIT_EVENT(FControlRigTrait_v2, NULL_ANIM_TRAIT_EVENT_ENUMERATOR)
+ANIM_NEXT_IMPL_DEFINE_ANIM_TRAIT_GET_TRAIT_EVENTS(FControlRigTrait_v2, NULL_ANIM_TRAIT_EVENT_ENUMERATOR)
 #undef TRAIT_INTERFACE_ENUMERATOR
 
 //----------------------
 // --- FInstanceData ---
 //----------------------
 
-void FControlRigTrait::FInstanceData::Construct(const FExecutionContext& Context, const FTraitBinding& Binding)
+template<typename TSharedData>
+void TControlRigTraitBase<TSharedData>::FInstanceData::Construct(const FExecutionContext& Context, const FTraitBinding& Binding)
 {
 	FTrait::FInstanceData::Construct(Context, Binding);
 
@@ -204,16 +293,17 @@ void FControlRigTrait::FInstanceData::Construct(const FExecutionContext& Context
 
 #if WITH_EDITOR
 	{
-		OnObjectsReinstancedHandle = UE::UAF::ControlRig::FAnimNextControlRigModule::OnObjectsReinstanced.AddRaw(this, &FControlRigTrait::FInstanceData::OnObjectsReinstanced);
+		OnObjectsReinstancedHandle = UE::UAF::ControlRig::FAnimNextControlRigModule::OnObjectsReinstanced.AddRaw(this, &FInstanceData::OnObjectsReinstanced);
 	}
 #endif
 
-	if (UClass* ControlRigClass = GetTargetClass(SharedData))
+	FControlRigAssetStrongReference CRAssetReference = GetTargetAssetReference(SharedData);
+	if (CRAssetReference.IsValid())
 	{
 		if (UObject* AnimContext = GetAnimContext(Context))
 		{
 			USkeletalMeshComponent* BindableObject = const_cast<USkeletalMeshComponent*>(FInstanceData::GetBindableObject(Context));
-			if (ensure(CreateControlRig(GetAnimContext(Context), BindableObject, SharedData->ControlRigClass, this)))
+			if (ensure(CreateControlRig(GetAnimContext(Context), BindableObject, SharedData->GetControlRigAssetReference(), this)))
 			{
 				FAnimNextControlRigTaskParams Params;
 				Params.ControlRig = ControlRig;
@@ -237,7 +327,8 @@ void FControlRigTrait::FInstanceData::Construct(const FExecutionContext& Context
 	}
 }
 
-void FControlRigTrait::FInstanceData::Destruct(const FExecutionContext& Context, const FTraitBinding& Binding)
+template<typename TSharedData>
+void TControlRigTraitBase<TSharedData>::FInstanceData::Destruct(const FExecutionContext& Context, const FTraitBinding& Binding)
 {
 	FTrait::FInstanceData::Destruct(Context, Binding);
 
@@ -250,7 +341,8 @@ void FControlRigTrait::FInstanceData::Destruct(const FExecutionContext& Context,
 	DestroyControlRig(Context, Binding);
 }
 
-void FControlRigTrait::FInstanceData::InitializeControlRig(const FExecutionContext& Context, const FTraitBinding& Binding)
+template<typename TSharedData>
+void TControlRigTraitBase<TSharedData>::FInstanceData::InitializeControlRig(const FExecutionContext& Context, const FTraitBinding& Binding)
 {
 	const FSharedData* SharedData = Binding.GetSharedData<FSharedData>();
 	check(EvaluationTask);
@@ -259,11 +351,12 @@ void FControlRigTrait::FInstanceData::InitializeControlRig(const FExecutionConte
 		EvaluationTask->Initialize(true, &PropertyMappings);
 
 		ControlRig->OnInitialized_AnyThread().RemoveAll(this);
-		OnInitializedHandle = ControlRig->OnInitialized_AnyThread().AddRaw(this, &FControlRigTrait::FInstanceData::HandleOnInitialized_AnyThread);
+		OnInitializedHandle = ControlRig->OnInitialized_AnyThread().AddRaw(this, &TControlRigTraitBase<TSharedData>::FInstanceData::HandleOnInitialized_AnyThread);
 	}
 }
 
-void FControlRigTrait::FInstanceData::DestroyControlRig(const FExecutionContext& Context, const FTraitBinding& Binding)
+template<typename TSharedData>
+void TControlRigTraitBase<TSharedData>::FInstanceData::DestroyControlRig(const FExecutionContext& Context, const FTraitBinding& Binding)
 {
 	if (ControlRig.Get() != nullptr)
 	{
@@ -277,7 +370,8 @@ void FControlRigTrait::FInstanceData::DestroyControlRig(const FExecutionContext&
 	}
 }
 
-UObject* FControlRigTrait::FInstanceData::GetAnimContext(const FExecutionContext& Context)
+template<typename TSharedData>
+UObject* TControlRigTraitBase<TSharedData>::FInstanceData::GetAnimContext(const FExecutionContext& Context)
 {
 	UObject* AnimContext = nullptr;
 
@@ -289,66 +383,67 @@ UObject* FControlRigTrait::FInstanceData::GetAnimContext(const FExecutionContext
 	return AnimContext;
 }
 
-const USkeletalMeshComponent* FControlRigTrait::FInstanceData::GetBindableObject(const FExecutionContext& Context)
+template<typename TSharedData>
+const USkeletalMeshComponent* TControlRigTraitBase<TSharedData>::FInstanceData::GetBindableObject(const FExecutionContext& Context)
 {
 	const USkeletalMeshComponent* BindableObject = Context.GetBindingObject().Get();
 	
 	return BindableObject;
 }
 
-UClass* FControlRigTrait::FInstanceData::GetTargetClass(const FSharedData* SharedData)
+template<typename TSharedData>
+FControlRigAssetStrongReference TControlRigTraitBase<TSharedData>::FInstanceData::GetTargetAssetReference(const FSharedData* SharedData)
 {
 	if (SharedData != nullptr)
 	{
-		return SharedData->ControlRigClass;
+		return SharedData->GetControlRigAssetReference();
 	}
 	return nullptr;
 }
 
-int32 FControlRigTrait::FInstanceData::GetExposedVariablesData(const UE::UAF::FTraitBinding& Binding, const FSharedData* SharedData, FControlRigVariableMappings::FCustomPropertyMappings& OutMappings)
+template<typename TSharedData>
+int32 TControlRigTraitBase<TSharedData>::FInstanceData::GetExposedVariablesData(const UE::UAF::FTraitBinding& Binding, const FSharedData* SharedData, FControlRigVariableMappings::FCustomPropertyMappings& OutMappings)
 {
 	int32 NumElementsAdded = 0;
 
-	if (SharedData && SharedData->ControlRigClass.Get() != nullptr && SharedData->ExposedPropertyVariableNames.Num() > 0)
+	if (SharedData && SharedData->HasValidControlRigReference() && SharedData->ExposedPropertyVariableNames.Num() > 0)
 	{
-		if (const UControlRig* CDO = SharedData->ControlRigClass->GetDefaultObject<UControlRig>())
+		FControlRigAssetStrongReference CRAsset = SharedData->GetControlRigAssetReference();
+		const int32 NumLatentProperties = SharedData->LatentProperties.Num();
+		if (ensure(NumLatentProperties == SharedData->LatentPropertyMemoryLayouts.Num()))
 		{
-			const int32 NumLatentProperties = SharedData->LatentProperties.Num();
-			if (ensure(NumLatentProperties == SharedData->LatentPropertyMemoryLayouts.Num()))
+			const TArray<FRigVMExternalVariable> PublicVariables = CRAsset.GetExternalVariables();
+			for (int32 LatentPropertyIndex = 0; LatentPropertyIndex < NumLatentProperties; LatentPropertyIndex++)
 			{
-				const TArray<FRigVMExternalVariable> PublicVariables = CDO->GetPublicVariables();
-				for (int32 LatentPropertyIndex = 0; LatentPropertyIndex < NumLatentProperties; LatentPropertyIndex++)
+				const FName& LatentPropertyName = SharedData->LatentProperties[LatentPropertyIndex];
+
+				if (!SharedData->ExposedPropertyVariableNames.Contains(LatentPropertyName))	// Only process exposed public variables
 				{
-					const FName& LatentPropertyName = SharedData->LatentProperties[LatentPropertyIndex];
+					continue;
+				}
 
-					if (!SharedData->ExposedPropertyVariableNames.Contains(LatentPropertyName))	// Only process exposed public variables
+				const FRigVMExternalVariable* Variable = PublicVariables.FindByPredicate([&LatentPropertyName](const FRigVMExternalVariable& In)
 					{
-						continue;
-					}
+						return In.GetName() == LatentPropertyName;
+					});
 
-					const FRigVMExternalVariable* Variable = PublicVariables.FindByPredicate([&LatentPropertyName](const FRigVMExternalVariable& In)
-						{
-							return In.GetName() == LatentPropertyName;
-						});
+				if (Variable != nullptr)
+				{
+					uint32 PropertyAlignment = 0;
+					uint32 PropertySize = 0;
+					UE::UAF::TControlRigTraitBase<TSharedData>::GetVariableSizeAndAlignment(*Variable, PropertySize, PropertyAlignment);
+					check(PropertyAlignment < MAX_uint16);
+					check(PropertySize < MAX_uint16);
 
-					if (Variable != nullptr)
+					if (ensure(SharedData->LatentPropertyMemoryLayouts[LatentPropertyIndex] == ((PropertySize << 16) | PropertyAlignment)))
 					{
-						uint32 PropertyAlignment = 0;
-						uint32 PropertySize = 0;
-						UE::UAF::FControlRigTrait::GetVariableSizeAndAlignment(*Variable, PropertySize, PropertyAlignment);
-						check(PropertyAlignment < MAX_uint16);
-						check(PropertySize < MAX_uint16);
+						const UE::UAF::FLatentPropertyHandle* TraitLatentPropertyHandles = Binding.GetLatentPropertyHandles();
 
-						if (ensure(SharedData->LatentPropertyMemoryLayouts[LatentPropertyIndex] == ((PropertySize << 16) | PropertyAlignment)))
+						const UE::UAF::FLatentPropertyHandle& LatentPropertyHandle = TraitLatentPropertyHandles[LatentPropertyIndex];
+						if (LatentPropertyHandle.IsOffsetValid())
 						{
-							const UE::UAF::FLatentPropertyHandle* TraitLatentPropertyHandles = Binding.GetLatentPropertyHandles();
-
-							const UE::UAF::FLatentPropertyHandle& LatentPropertyHandle = TraitLatentPropertyHandles[LatentPropertyIndex];
-							if (LatentPropertyHandle.IsOffsetValid())
-							{
-								const uint8* LatentPropertyMemory = const_cast<uint8*>(Binding.GetLatentProperty<uint8>(LatentPropertyHandle));
-								OutMappings.AddVariable(Variable->GetName(), const_cast<uint8*>(Variable->GetMemory()), Variable->GetProperty(), LatentPropertyMemory);
-							}
+							const uint8* LatentPropertyMemory = const_cast<uint8*>(Binding.GetLatentProperty<uint8>(LatentPropertyHandle));
+							OutMappings.AddVariable(Variable->GetName(), const_cast<uint8*>(Variable->GetMemory()), Variable->GetProperty(), LatentPropertyMemory);
 						}
 					}
 				}
@@ -359,11 +454,12 @@ int32 FControlRigTrait::FInstanceData::GetExposedVariablesData(const UE::UAF::FT
 	return NumElementsAdded;
 }
 
-int32 FControlRigTrait::FInstanceData::GetExposedControlsData(const UE::UAF::FTraitBinding& Binding, const FSharedData* SharedData, FControlRigVariableMappings::FCustomPropertyMappings& OutMappings)
+template<typename TSharedData>
+int32 TControlRigTraitBase<TSharedData>::FInstanceData::GetExposedControlsData(const UE::UAF::FTraitBinding& Binding, const FSharedData* SharedData, FControlRigVariableMappings::FCustomPropertyMappings& OutMappings)
 {
 	int32 NumElementsAdded = 0;
 
-	if (SharedData && SharedData->ControlRigClass.Get() != nullptr)
+	if (SharedData && SharedData->HasValidControlRigReference())
 	{
 		const int32 NumLatentProperties = SharedData->LatentProperties.Num();
 		if (ensure(NumLatentProperties == SharedData->LatentPropertyMemoryLayouts.Num()))
@@ -387,7 +483,7 @@ int32 FControlRigTrait::FInstanceData::GetExposedControlsData(const UE::UAF::FTr
 
 					uint32 PropertyAlignment = 0;
 					uint32 PropertySize = 0;
-					UE::UAF::FControlRigTrait::GetControlSizeAndAlignment(ControlType, PropertySize, PropertyAlignment);
+					UE::UAF::TControlRigTraitBase<TSharedData>::GetControlSizeAndAlignment(ControlType, PropertySize, PropertyAlignment);
 					ensure(PropertyAlignment < MAX_uint16);
 					ensure(PropertySize < MAX_uint16);
 
@@ -410,12 +506,14 @@ int32 FControlRigTrait::FInstanceData::GetExposedControlsData(const UE::UAF::FTr
 	return NumElementsAdded;
 }
 
-void FControlRigTrait::FInstanceData::MapUAFVariablesToControlRig(const UE::UAF::FTraitBinding& Binding, const FSharedData* SharedData,
+template<typename TSharedData>
+void TControlRigTraitBase<TSharedData>::FInstanceData::MapUAFVariablesToControlRig(const UE::UAF::FTraitBinding& Binding, const FSharedData* SharedData,
 	FControlRigVariableMappings::FCustomPropertyMappings& OutMappings)
 {
 	FUAFInstanceVariableData& Variables = Binding.GetTraitPtr().GetNodeInstance()->GetOwner().GetVariables();
 
-	if (const IRigVMRuntimeAssetInterface* ControlRigAsset = Cast<const IRigVMRuntimeAssetInterface>(SharedData->ControlRigClass.Get()))
+	FControlRigAssetStrongReference CRAssetReference = SharedData->GetControlRigAssetReference();
+	if (const IRigVMRuntimeAssetInterface* ControlRigAsset = Cast<const IRigVMRuntimeAssetInterface>(CRAssetReference.Get()))
 	{
 		// Todo: this is a copy so its slow
 		// Ideally we can get the details of the control rig without a copy
@@ -460,7 +558,8 @@ void FControlRigTrait::FInstanceData::MapUAFVariablesToControlRig(const UE::UAF:
 	
 }
 
-void FControlRigTrait::FInstanceData::HandleOnInitialized_AnyThread(URigVMHost*, const FName&)
+template<typename TSharedData>
+void TControlRigTraitBase<TSharedData>::FInstanceData::HandleOnInitialized_AnyThread(URigVMHost*, const FName&)
 {
 	check(EvaluationTask);
 	EvaluationTask->OnControlRigInitialized();
@@ -470,7 +569,8 @@ void FControlRigTrait::FInstanceData::HandleOnInitialized_AnyThread(URigVMHost*,
 	#endif
 }
 
-void FControlRigTrait::FInstanceData::InitializeCustomProperties(const FTraitBinding& Binding, const FSharedData* SharedData)
+template<typename TSharedData>
+void TControlRigTraitBase<TSharedData>::FInstanceData::InitializeCustomProperties(const FTraitBinding& Binding, const FSharedData* SharedData)
 {
 	check(EvaluationTask);
 	// Setup mappings using the latent pin memory as source (we have to copy from latent pin to external variable / rig control)
@@ -478,7 +578,8 @@ void FControlRigTrait::FInstanceData::InitializeCustomProperties(const FTraitBin
 }
 
 #if WITH_EDITOR
-void FControlRigTrait::FInstanceData::OnObjectsReinstanced(const TMap<UObject*, UObject*>& OldToNewInstanceMap)
+template<typename TSharedData>
+void TControlRigTraitBase<TSharedData>::FInstanceData::OnObjectsReinstanced(const TMap<UObject*, UObject*>& OldToNewInstanceMap)
 {
 	if (ControlRig)
 	{
@@ -505,16 +606,17 @@ void FControlRigTrait::FInstanceData::OnObjectsReinstanced(const TMap<UObject*,
 // --- FControlRigTrait ---
 //-------------------------
 // --- Custom Latent Pin Support impl ---
-void FControlRigTrait::SerializeTraitSharedData(FArchive& Ar, FAnimNextTraitSharedData& SharedData) const
+template<typename TSharedData>
+void TControlRigTraitBase<TSharedData>::SerializeTraitSharedData(FArchive& Ar, FAnimNextTraitSharedData& SharedData) const
 {
-	FControlRigTraitSharedData& ControlRigSharedData = static_cast<FControlRigTraitSharedData&>(SharedData);
+	TSharedData& ControlRigSharedData = static_cast<TSharedData&>(SharedData);
 
 	TraitSuper::SerializeTraitSharedData(Ar, SharedData);
 
 	if (Ar.IsLoading())
 	{
 		// Compute our latent property data based on our ControlRig class
-		if (UClass* ControlRigClass = ControlRigSharedData.ControlRigClass)
+		if (ControlRigSharedData.HasValidControlRigReference())
 		{
 			// We build the size/alignment map for each property even if their pin isn't hooked to anything
 			// since handles are reserved for every one of them
@@ -524,31 +626,28 @@ void FControlRigTrait::SerializeTraitSharedData(FArchive& Ar, FAnimNextTraitShar
 
 			if (ControlRigSharedData.ExposedPropertyVariableNames.Num() > 0)
 			{
-				if (const UControlRig* CDO = ControlRigClass->GetDefaultObject<UControlRig>())
+				const TArray<FRigVMExternalVariable> PublicVariables = ControlRigSharedData.GetControlRigAssetReference().GetExternalVariables();
+				if (!ensureMsgf(PublicVariables.Num() <= (int32)GetNumLatentTraitProperties(), TEXT("The ControlRig Trait only supports up to %u input variables"), GetNumLatentTraitProperties()))
 				{
-					const TArray<FRigVMExternalVariable> PublicVariables = CDO->GetPublicVariables();
-					if (!ensureMsgf(PublicVariables.Num() <= (int32)GetNumLatentTraitProperties(), TEXT("The ControlRig Trait only supports up to %u input variables"), GetNumLatentTraitProperties()))
+					return;
+				}
+
+				for (const FRigVMExternalVariable& Variable : PublicVariables)
+				{
+					if (!ControlRigSharedData.ExposedPropertyVariableNames.Contains(Variable.GetName()))	// Only process exposed public variables
 					{
-						return;
+						continue;
 					}
 
-					for (const FRigVMExternalVariable& Variable : PublicVariables)
+					uint32 PropertyAlignment = 0;
+					uint32 PropertySize = 0;
+					if (GetVariableSizeAndAlignment(Variable, PropertySize, PropertyAlignment))
 					{
-						if (!ControlRigSharedData.ExposedPropertyVariableNames.Contains(Variable.GetName()))	// Only process exposed public variables
-						{
-							continue;
-						}
-
-						uint32 PropertyAlignment = 0;
-						uint32 PropertySize = 0;
-						if (GetVariableSizeAndAlignment(Variable, PropertySize, PropertyAlignment))
-						{
-							check(PropertyAlignment < MAX_uint16);
-							check(PropertySize < MAX_uint16);
+						check(PropertyAlignment < MAX_uint16);
+						check(PropertySize < MAX_uint16);
 
-							ControlRigSharedData.LatentProperties.Add(Variable.GetName());
-							ControlRigSharedData.LatentPropertyMemoryLayouts.Add((PropertySize << 16) | PropertyAlignment);
-						}
+						ControlRigSharedData.LatentProperties.Add(Variable.GetName());
+						ControlRigSharedData.LatentPropertyMemoryLayouts.Add((PropertySize << 16) | PropertyAlignment);
 					}
 				}
 			}
@@ -577,16 +676,18 @@ void FControlRigTrait::SerializeTraitSharedData(FArchive& Ar, FAnimNextTraitShar
 	}
 }
 
-uint32 FControlRigTrait::GetNumLatentTraitProperties() const
+template<typename TSharedData>
+uint32 TControlRigTraitBase<TSharedData>::GetNumLatentTraitProperties() const
 {
 	// Number of latent trait properties must be known ahead of time to reserve space
 	// We support a maximum number of input properties, each one will need a 2-byte handle in the shared data for each trait
 	return 64;
 }
 
-FTraitLatentPropertyMemoryLayout FControlRigTrait::GetLatentPropertyMemoryLayout(const FAnimNextTraitSharedData& SharedData, FName PropertyName, uint32 PropertyIndex) const
+template<typename TSharedData>
+FTraitLatentPropertyMemoryLayout TControlRigTraitBase<TSharedData>::GetLatentPropertyMemoryLayout(const FAnimNextTraitSharedData& SharedData, FName PropertyName, uint32 PropertyIndex) const
 {
-	const FControlRigTraitSharedData& ControlRigSharedData = static_cast<const FControlRigTraitSharedData&>(SharedData);
+	const TSharedData& ControlRigSharedData = static_cast<const TSharedData&>(SharedData);
 
 	const int32 LatentPropertyIndex = ControlRigSharedData.LatentProperties.IndexOfByKey(PropertyName);
 	if (LatentPropertyIndex == INDEX_NONE)
@@ -603,7 +704,8 @@ FTraitLatentPropertyMemoryLayout FControlRigTrait::GetLatentPropertyMemoryLayout
 }
 
 // --- IUpdate impl --- 
-void FControlRigTrait::OnBecomeRelevant(FUpdateTraversalContext& Context, const TTraitBinding<IUpdate>& Binding, const FTraitUpdateState& TraitState) const
+template<typename TSharedData>
+void TControlRigTraitBase<TSharedData>::OnBecomeRelevant(FUpdateTraversalContext& Context, const TTraitBinding<IUpdate>& Binding, const FTraitUpdateState& TraitState) const
 {
 	const FSharedData* SharedData = Binding.GetSharedData<FSharedData>();
 	check(SharedData);
@@ -616,7 +718,8 @@ void FControlRigTrait::OnBecomeRelevant(FUpdateTraversalContext& Context, const
 	}
 }
 
-void FControlRigTrait::PreUpdate(FUpdateTraversalContext& Context, const TTraitBinding<IUpdate>& Binding, const FTraitUpdateState& TraitState) const
+template<typename TSharedData>
+void TControlRigTraitBase<TSharedData>::PreUpdate(FUpdateTraversalContext& Context, const TTraitBinding<IUpdate>& Binding, const FTraitUpdateState& TraitState) const
 {
 	IUpdate::PreUpdate(Context, Binding, TraitState);
 
@@ -638,9 +741,9 @@ void FControlRigTrait::PreUpdate(FUpdateTraversalContext& Context, const TTraitB
 	if (InstanceData->bReinitializeControlRig)
 	{
 		InstanceData->PropertyMappings.Reset();
-		UE::UAF::FControlRigTrait::FInstanceData::GetExposedVariablesData(Binding, SharedData, InstanceData->PropertyMappings);
-		UE::UAF::FControlRigTrait::FInstanceData::GetExposedControlsData(Binding, SharedData, InstanceData->PropertyMappings);
-		UE::UAF::FControlRigTrait::FInstanceData::MapUAFVariablesToControlRig(Binding, SharedData, InstanceData->PropertyMappings);
+		UE::UAF::TControlRigTraitBase<TSharedData>::FInstanceData::GetExposedVariablesData(Binding, SharedData, InstanceData->PropertyMappings);
+		UE::UAF::TControlRigTraitBase<TSharedData>::FInstanceData::GetExposedControlsData(Binding, SharedData, InstanceData->PropertyMappings);
+		UE::UAF::TControlRigTraitBase<TSharedData>::FInstanceData::MapUAFVariablesToControlRig(Binding, SharedData, InstanceData->PropertyMappings);
 
 		InstanceData->InitializeControlRig(Context, Binding);
 		InstanceData->bReinitializeControlRig = false;
@@ -665,7 +768,8 @@ void FControlRigTrait::PreUpdate(FUpdateTraversalContext& Context, const TTraitB
 }
 
 // --- IEvaluate impl --- 
-void FControlRigTrait::PostEvaluate(FEvaluateTraversalContext& Context, const TTraitBinding<IEvaluate>& Binding) const
+template<typename TSharedData>
+void TControlRigTraitBase<TSharedData>::PostEvaluate(FEvaluateTraversalContext& Context, const TTraitBinding<IEvaluate>& Binding) const
 {
 	const FSharedData* SharedData = Binding.GetSharedData<FSharedData>();
 	FInstanceData* InstanceData = Binding.GetInstanceData<FInstanceData>();
@@ -688,19 +792,22 @@ void FControlRigTrait::PostEvaluate(FEvaluateTraversalContext& Context, const TT
 }
 
 // --- IHierarchy impl --- 
-uint32 FControlRigTrait::GetNumChildren(const FExecutionContext& Context, const TTraitBinding<IHierarchy>& Binding) const
+template<typename TSharedData>
+uint32 TControlRigTraitBase<TSharedData>::GetNumChildren(const FExecutionContext& Context, const TTraitBinding<IHierarchy>& Binding) const
 {
 	return 1;
 }
 
-void FControlRigTrait::GetChildren(const FExecutionContext& Context, const TTraitBinding<IHierarchy>& Binding, FChildrenArray& Children) const
+template<typename TSharedData>
+void TControlRigTraitBase<TSharedData>::GetChildren(const FExecutionContext& Context, const TTraitBinding<IHierarchy>& Binding, FChildrenArray& Children) const
 {
 	const FInstanceData* InstanceData = Binding.GetInstanceData<FInstanceData>();
 	Children.Add(InstanceData->Input);
 }
 
 // --- IGarbageCollection impl --- 
-void FControlRigTrait::AddReferencedObjects(const FExecutionContext& Context, const TTraitBinding<IGarbageCollection>& Binding, FReferenceCollector& Collector) const
+template<typename TSharedData>
+void TControlRigTraitBase<TSharedData>::AddReferencedObjects(const FExecutionContext& Context, const TTraitBinding<IGarbageCollection>& Binding, FReferenceCollector& Collector) const
 {
 	IGarbageCollection::AddReferencedObjects(Context, Binding, Collector);
 
@@ -715,7 +822,8 @@ void FControlRigTrait::AddReferencedObjects(const FExecutionContext& Context, co
 
 // --- Utility --- 
 
-UControlRig* FControlRigTrait::GetControlRig(FInstanceData* InstanceData)
+template<typename TSharedData>
+UControlRig* TControlRigTraitBase<TSharedData>::GetControlRig(FInstanceData* InstanceData)
 {
 	if (InstanceData)
 	{
@@ -725,16 +833,17 @@ UControlRig* FControlRigTrait::GetControlRig(FInstanceData* InstanceData)
 	return nullptr;
 }
 
-bool FControlRigTrait::CreateControlRig(UObject* InAnimContext, UObject* InBindableObject, UClass* InControlRigClass, FInstanceData* InstanceData)
+template<typename TSharedData>
+bool TControlRigTraitBase<TSharedData>::CreateControlRig(UObject* InAnimContext, UObject* InBindableObject, FControlRigAssetStrongReference InControlRigAssetReference, FInstanceData* InstanceData)
 {
-	if (InstanceData->ControlRig == nullptr && InControlRigClass != nullptr)
+	if (InstanceData->ControlRig == nullptr && InControlRigAssetReference.IsValid())
 	{
 		bool bSuccess = false;
 		// Let's make sure the GC isn't running when we try to create a new Control Rig 
 		// or when setting the bindable object as GetDataSourceRegistry might create a new object
 		{
 			FGCScopeGuard GCGuard;
-			InstanceData->ControlRig = NewObject<UControlRig>(InAnimContext, InControlRigClass);
+			InstanceData->ControlRig = InControlRigAssetReference.CreateInstance(InAnimContext);
 			InstanceData->ControlRig->AtomicallyClearInternalFlags(EInternalObjectFlags::Async);
 			
 			bSuccess = SetBindableObject(InstanceData->ControlRig, InAnimContext, InBindableObject);
@@ -749,7 +858,8 @@ bool FControlRigTrait::CreateControlRig(UObject* InAnimContext, UObject* InBinda
 }
 
 
-bool FControlRigTrait::SetBindableObject(UControlRig* ControlRig, UObject* InAnimContext, UObject* InBindableObject)
+template<typename TSharedData>
+bool TControlRigTraitBase<TSharedData>::SetBindableObject(UControlRig* ControlRig, UObject* InAnimContext, UObject* InBindableObject)
 {
 	if (ControlRig != nullptr && ensure(InAnimContext) && InBindableObject)
 	{
@@ -769,44 +879,42 @@ bool FControlRigTrait::SetBindableObject(UControlRig* ControlRig, UObject* InAni
 }
 
 #if WITH_EDITOR
-void FControlRigTrait::GetProgrammaticPins(FAnimNextTraitSharedData* InSharedData, URigVMController* InController, int32 InParentPinIndex, const URigVMPin* InTraitPin, const FString& InDefaultValue, struct FRigVMPinInfoArray& OutPinArray) const
+template<typename TSharedData>
+void TControlRigTraitBase<TSharedData>::GetProgrammaticPins(FAnimNextTraitSharedData* InSharedData, URigVMController* InController, int32 InParentPinIndex, const URigVMPin* InTraitPin, const FString& InDefaultValue, struct FRigVMPinInfoArray& OutPinArray) const
 {
-	if (FControlRigTraitSharedData* SharedData = (FControlRigTraitSharedData*)InSharedData)
+	if (TSharedData* SharedData = (TSharedData*)InSharedData)
 	{
-		if (SharedData->ControlRigClass != nullptr)
+		if (SharedData->HasValidControlRigReference())
 		{
 			// --- Exposed Public Variables ---
 			if (SharedData->ExposedPropertyVariableNames.Num() > 0)
 			{
-				if (const UControlRig* CDO = SharedData->ControlRigClass->GetDefaultObject<UControlRig>())
-				{
-					const TArray<FRigVMExternalVariable> PublicVariables = CDO->GetPublicVariables();
+				const TArray<FRigVMExternalVariable> PublicVariables = SharedData->GetControlRigAssetReference().GetExternalVariables();
 
-					for (const FRigVMExternalVariable& Variable : PublicVariables)
+				for (const FRigVMExternalVariable& Variable : PublicVariables)
+				{
+					if (!SharedData->ExposedPropertyVariableNames.Contains(Variable.GetName()))	// Only process exposed public variables
 					{
-						if (!SharedData->ExposedPropertyVariableNames.Contains(Variable.GetName()))	// Only process exposed public variables
-						{
-							continue;
-						}
+						continue;
+					}
 
-						if (Variable.GetMemory() == nullptr) // if we make a variable public but don't recompile, we have to skip, as it comes without memory
-						{
-							continue;
-						}
+					if (Variable.GetMemory() == nullptr) // if we make a variable public but don't recompile, we have to skip, as it comes without memory
+					{
+						continue;
+					}
 
-						if (const URigVMPin* SubPin = InTraitPin->FindSubPin(Variable.GetName().ToString()))
-						{
-							const FString PinDefaultValue = SubPin->GetDefaultValue();
-							FRigVMMemoryStorageStruct StorageDefaultValue(ERigVMMemoryType::External, {FRigVMPropertyDescription(Variable.GetProperty(), PinDefaultValue, Variable.GetName())});
+					if (const URigVMPin* SubPin = InTraitPin->FindSubPin(Variable.GetName().ToString()))
+					{
+						const FString PinDefaultValue = SubPin->GetDefaultValue();
+						FRigVMMemoryStorageStruct StorageDefaultValue(ERigVMMemoryType::External, {FRigVMPropertyDescription(Variable.GetProperty(), PinDefaultValue, Variable.GetName())});
 
-							const uint8* DefaultMemory = StorageDefaultValue.GetDataByName<const uint8>(Variable.GetName());
-							
-							OutPinArray.AddPin(const_cast<FProperty*>(Variable.GetProperty()), InController, ERigVMPinDirection::Input, InParentPinIndex, ERigVMPinDefaultValueType::AutoDetect, DefaultMemory, true);
-						}
-						else
-						{
-							OutPinArray.AddPin(const_cast<FProperty*>(Variable.GetProperty()), InController, ERigVMPinDirection::Input, InParentPinIndex, ERigVMPinDefaultValueType::AutoDetect, Variable.GetMemory(), true);
-						}
+						const uint8* DefaultMemory = StorageDefaultValue.GetDataByName<const uint8>(Variable.GetName());
+						
+						OutPinArray.AddPin(const_cast<FProperty*>(Variable.GetProperty()), InController, ERigVMPinDirection::Input, InParentPinIndex, ERigVMPinDefaultValueType::AutoDetect, DefaultMemory, true);
+					}
+					else
+					{
+						OutPinArray.AddPin(const_cast<FProperty*>(Variable.GetProperty()), InController, ERigVMPinDirection::Input, InParentPinIndex, ERigVMPinDefaultValueType::AutoDetect, Variable.GetMemory(), true);
 					}
 				}
 			}
@@ -834,7 +942,7 @@ void FControlRigTrait::GetProgrammaticPins(FAnimNextTraitSharedData* InSharedDat
 					if (USkeleton* PreviewSkeleton = SharedData->GetPreviewSkeleton())
 					{
 						// Obtain the controls from the RigControlsData helper. This will instantiate a rig using the provided class and cache the controls until the class changes
-						const TArray<FControlRigIOMapping::FControlsInfo>& RigControls = RigControlsData.GetControls(SharedData->ControlRigClass, PreviewSkeleton);
+						const TArray<FControlRigIOMapping::FControlsInfo>& RigControls = RigControlsData.GetControls(SharedData->GetControlRigAssetReference(), PreviewSkeleton);
 
 						for (const FControlRigIOMapping::FControlsInfo& Control : RigControls)
 						{
@@ -858,13 +966,15 @@ void FControlRigTrait::GetProgrammaticPins(FAnimNextTraitSharedData* InSharedDat
 	}
 }
 
-int32 FControlRigTrait::GetLatentPropertyIndex(const FAnimNextTraitSharedData& InSharedData, FName PropertyName) const
+template<typename TSharedData>
+int32 TControlRigTraitBase<TSharedData>::GetLatentPropertyIndex(const FAnimNextTraitSharedData& InSharedData, FName PropertyName) const
 {
-	const FControlRigTraitSharedData& SharedData = (const FControlRigTraitSharedData&)InSharedData;
+	const TSharedData& SharedData = (const TSharedData&)InSharedData;
 	return SharedData.LatentProperties.IndexOfByKey(PropertyName);
 }
 
-uint32 FControlRigTrait::GetLatentPropertyHandles(
+template<typename TSharedData>
+uint32 TControlRigTraitBase<TSharedData>::GetLatentPropertyHandles(
 	const FAnimNextTraitSharedData* InSharedData,
 	TArray<FLatentPropertyMetadata>& OutLatentPropertyHandles,
 	bool bFilterEditorOnly,
@@ -874,36 +984,32 @@ uint32 FControlRigTrait::GetLatentPropertyHandles(
 	uint32 NumLatentPinsAdded = Super::GetLatentPropertyHandles(InSharedData, OutLatentPropertyHandles, bFilterEditorOnly, GetTraitLatentPropertyIndex);
 
 	// Generate Control Rig exposed pins
-	if (FControlRigTraitSharedData* SharedData = (FControlRigTraitSharedData*)InSharedData)
+	if (TSharedData* SharedData = (TSharedData*)InSharedData)
 	{
-		if (SharedData->ControlRigClass != nullptr)
+		if (SharedData->HasValidControlRigReference())
 		{
 			// --- Iterate over public variables --- 
 			if (SharedData->ExposedPropertyVariableNames.Num() > 0)
 			{
-				if (const UControlRig* CDO = SharedData->ControlRigClass->GetDefaultObject<UControlRig>())
+				const TArray<FRigVMExternalVariable> PublicVariables = SharedData->GetControlRigAssetReference().GetExternalVariables();
+				for (const FRigVMExternalVariable& Variable : PublicVariables)
 				{
-					const TArray<FRigVMExternalVariable> PublicVariables = CDO->GetPublicVariables();
-
-					for (const FRigVMExternalVariable& Variable : PublicVariables)
+					if (!SharedData->ExposedPropertyVariableNames.Contains(Variable.GetName()))	// Only process exposed public variables
 					{
-						if (!SharedData->ExposedPropertyVariableNames.Contains(Variable.GetName()))	// Only process exposed public variables
-						{
-							continue;
-						}
+						continue;
+					}
 
-						const FProperty* Property = Variable.GetProperty();
+					const FProperty* Property = Variable.GetProperty();
 
-						FLatentPropertyMetadata Metadata;
-						Metadata.Name = Property->GetFName();
-						Metadata.RigVMIndex = GetTraitLatentPropertyIndex(Property->GetFName());
+					FLatentPropertyMetadata Metadata;
+					Metadata.Name = Property->GetFName();
+					Metadata.RigVMIndex = GetTraitLatentPropertyIndex(Property->GetFName());
 
-						// Always false for now, we don't support freezing yet
-						Metadata.bCanFreeze = false;
+					// Always false for now, we don't support freezing yet
+					Metadata.bCanFreeze = false;
 
-						OutLatentPropertyHandles.Add(Metadata);
-						NumLatentPinsAdded++;
-					}
+					OutLatentPropertyHandles.Add(Metadata);
+					NumLatentPinsAdded++;
 				}
 			}
 
@@ -913,7 +1019,7 @@ uint32 FControlRigTrait::GetLatentPropertyHandles(
 				if (USkeleton* PreviewSkeleton = SharedData->GetPreviewSkeleton())
 				{
 					// Obtain the controls from the RigControlsData helper. This will instantiate a rig using the provided class and cache the controls until the class changes
-					const TArray<FControlRigIOMapping::FControlsInfo>& RigControls = RigControlsData.GetControls(SharedData->ControlRigClass, PreviewSkeleton);
+					const TArray<FControlRigIOMapping::FControlsInfo>& RigControls = RigControlsData.GetControls(SharedData->GetControlRigAssetReference(), PreviewSkeleton);
 
 					for (const FControlRigIOMapping::FControlsInfo& Control : RigControls)
 					{
@@ -943,7 +1049,8 @@ uint32 FControlRigTrait::GetLatentPropertyHandles(
 
 #endif // WITH_EDITOR
 
-bool FControlRigTrait::GetVariableSizeAndAlignment(const FRigVMExternalVariable& Variable, uint32& PropertySize, uint32& PropertyAlignment)
+template<typename TSharedData>
+bool TControlRigTraitBase<TSharedData>::GetVariableSizeAndAlignment(const FRigVMExternalVariable& Variable, uint32& PropertySize, uint32& PropertyAlignment)
 {
 	bool bValidType = false;
 
@@ -984,7 +1091,8 @@ bool FControlRigTrait::GetVariableSizeAndAlignment(const FRigVMExternalVariable&
 	return bValidType;
 }
 
-bool FControlRigTrait::GetControlSizeAndAlignment(ERigControlType ControlType, uint32& PropertySize, uint32& PropertyAlignment)
+template<typename TSharedData>
+bool TControlRigTraitBase<TSharedData>::GetControlSizeAndAlignment(ERigControlType ControlType, uint32& PropertySize, uint32& PropertyAlignment)
 {
 	switch (ControlType)
 	{
diff --git a/Engine/Plugins/Experimental/UAF/UAFControlRig/Source/UAFControlRig/Private/UAFControlRigAssetData.cpp b/Engine/Plugins/Experimental/UAF/UAFControlRig/Source/UAFControlRig/Private/UAFControlRigAssetData.cpp
--- a/Engine/Plugins/Experimental/UAF/UAFControlRig/Source/UAFControlRig/Private/UAFControlRigAssetData.cpp
+++ b/Engine/Plugins/Experimental/UAF/UAFControlRig/Source/UAFControlRig/Private/UAFControlRigAssetData.cpp
@@ -6,10 +6,21 @@
 
 bool FUAFGraphFactoryAsset_ControlRig::Validate() const
 {
-	return ControlRigClass != nullptr;
+	return ControlRigAssetReference.IsValid();
 }
 
 void FUAFGraphFactoryAsset_ControlRig::GetObjectReferences(TArray<const UObject*>& OutReferencedObjects) const
 {
-	OutReferencedObjects.Add(ControlRigClass);
+	OutReferencedObjects.Add(ControlRigAssetReference.Get());
 }
+
+#if WITH_EDITORONLY_DATA
+void FUAFGraphFactoryAsset_ControlRig::PostSerialize(const FArchive& Ar)
+{
+	if (ControlRigClass_DEPRECATED)
+	{
+		ControlRigAssetReference.Set(ControlRigClass_DEPRECATED);
+		ControlRigClass_DEPRECATED = nullptr;
+	}
+}
+#endif
diff --git a/Engine/Plugins/Experimental/UAF/UAFControlRig/Source/UAFControlRig/Public/UAFControlRigAssetData.h b/Engine/Plugins/Experimental/UAF/UAFControlRig/Source/UAFControlRig/Public/UAFControlRigAssetData.h
--- a/Engine/Plugins/Experimental/UAF/UAFControlRig/Source/UAFControlRig/Public/UAFControlRigAssetData.h
+++ b/Engine/Plugins/Experimental/UAF/UAFControlRig/Source/UAFControlRig/Public/UAFControlRigAssetData.h
@@ -2,6 +2,7 @@
 
 #pragma once
 
+#include "ControlRigAssetReference.h"
 #include "CoreTypes.h"
 #include "Templates/SubclassOf.h"
 #include "UAF/UAFAssetData.h"
@@ -14,13 +15,36 @@ USTRUCT(DisplayName="Control Rig Asset")
 struct FUAFGraphFactoryAsset_ControlRig : public FUAFGraphFactoryAsset
 {
 	GENERATED_BODY()
-	FUAFGraphFactoryAsset_ControlRig(const TSubclassOf<UControlRig>& InControlRigClass) : ControlRigClass(InControlRigClass) {}
+	FUAFGraphFactoryAsset_ControlRig(const FControlRigAssetStrongReference& InControlRigReference) : ControlRigAssetReference(InControlRigReference) {}
 	FUAFGraphFactoryAsset_ControlRig() = default;
 
 	virtual bool Validate() const override;
 	
 	virtual void GetObjectReferences(TArray<const UObject*>& OutReferencedObjects) const override;
 
+#if WITH_EDITORONLY_DATA
+	/** The type of this animation ControlRig */
+	UPROPERTY(meta=(DeprecatedProperty))
+	TSubclassOf<UControlRig> ControlRigClass_DEPRECATED;
+#endif
+	
+	/** The type of this animation ControlRig */
 	UPROPERTY(EditAnywhere, Category = ControlRig)
-	TSubclassOf<UControlRig> ControlRigClass;
+	FControlRigAssetStrongReference ControlRigAssetReference;
+	
+#if WITH_EDITORONLY_DATA
+	void PostSerialize(const FArchive& Ar);
+#endif
+};
+
+
+#if WITH_EDITORONLY_DATA
+template<>
+struct TStructOpsTypeTraits<FUAFGraphFactoryAsset_ControlRig> : public TStructOpsTypeTraitsBase2<FUAFGraphFactoryAsset_ControlRig>
+{
+	enum
+	{
+		WithPostSerialize = true,
+	};
 };
+#endif
\ No newline at end of file
diff --git a/Engine/Plugins/Experimental/UAF/UAFControlRig/Source/UAFControlRigEditor/Internal/ControlRigTraitCustomization.h b/Engine/Plugins/Experimental/UAF/UAFControlRig/Source/UAFControlRigEditor/Internal/ControlRigTraitCustomization.h
--- a/Engine/Plugins/Experimental/UAF/UAFControlRig/Source/UAFControlRigEditor/Internal/ControlRigTraitCustomization.h
+++ b/Engine/Plugins/Experimental/UAF/UAFControlRig/Source/UAFControlRigEditor/Internal/ControlRigTraitCustomization.h
@@ -12,7 +12,7 @@ class IDetailCategoryBuilder;
 class UEdGraphNode;
 class URigVMNode;
 struct FRigVMVariableMappingInfo;
-struct FControlRigTraitSharedData;
+struct FControlRigTraitSharedDataBase;
 
 namespace UE::Workspace
 {
@@ -42,7 +42,8 @@ private:
 	USkeleton* GetTargetSkeleton() const;
 
 	static const FControlRigIOMapping::FControlsInfo* GetControlInfo(const TArray<FControlRigIOMapping::FControlsInfo>& Controls, const FName& ControlName);
-	static bool IsVariableProperty(FControlRigTraitSharedData* ControlRigTraitSharedData, const FName& PropertyName);
+	template<typename TSharedData>
+	static bool IsVariableProperty(TSharedData* ControlRigTraitSharedData, const FName& PropertyName);
 
 	void OnObjectsReinstanced(const TMap<UObject*, UObject*>& OldToNewInstanceMap);
 
diff --git a/Engine/Plugins/Experimental/UAF/UAFControlRig/Source/UAFControlRigEditor/Private/ControlRigTraitCustomization.cpp b/Engine/Plugins/Experimental/UAF/UAFControlRig/Source/UAFControlRigEditor/Private/ControlRigTraitCustomization.cpp
--- a/Engine/Plugins/Experimental/UAF/UAFControlRig/Source/UAFControlRigEditor/Private/ControlRigTraitCustomization.cpp
+++ b/Engine/Plugins/Experimental/UAF/UAFControlRig/Source/UAFControlRigEditor/Private/ControlRigTraitCustomization.cpp
@@ -53,7 +53,11 @@ void FControlRigTraitSharedDataCustomization::CustomizeHeader(TSharedRef<IProper
 	StructPropertyHandle = InPropertyHandle;
 	ScopedControlRigTraitSharedData = GetControlRigSharedData(StructPropertyHandle);
 
-	FControlRigTraitSharedData* ControlRigTraitSharedData = (FControlRigTraitSharedData*)ScopedControlRigTraitSharedData->GetStructMemory();
+	FControlRigTraitSharedDataBase* ControlRigTraitSharedData = (FControlRigTraitSharedDataBase*)ScopedControlRigTraitSharedData->GetStructMemory();
+	if (!ControlRigTraitSharedData)
+	{
+		return;
+	}
 
 	CustomPinProperties.Reset(ControlRigTraitSharedData->ExposedPropertyVariableNames.Num() + ControlRigTraitSharedData->ExposedPropertyControlNames.Num());
 	for (const FName& PropertyName : ControlRigTraitSharedData->ExposedPropertyVariableNames)
@@ -107,7 +111,7 @@ void FControlRigTraitSharedDataCustomization::CustomizeChildren(TSharedRef<IProp
 
 void FControlRigTraitSharedDataCustomization::OnPropertyExposeCheckboxChanged(ECheckBoxState NewState, FName PropertyName)
 {
-	if (FControlRigTraitSharedData* ControlRigTraitSharedData = (FControlRigTraitSharedData*)ScopedControlRigTraitSharedData->GetStructMemory())
+	if (FControlRigTraitSharedDataBase* ControlRigTraitSharedData = (FControlRigTraitSharedDataBase*)ScopedControlRigTraitSharedData->GetStructMemory())
 	{
 		const bool bInput = ControlRigIOMapping->IsInputProperty(PropertyName);
 		ControlRigIOMapping->SetIOMapping(bInput, PropertyName, NAME_None);
@@ -122,10 +126,10 @@ void FControlRigTraitSharedDataCustomization::OnPropertyExposeCheckboxChanged(EC
 
 					if (URigVMEdGraph* EdGraph = Cast<URigVMEdGraph>(SelectedNodeWeak->GetGraph()))
 					{
-						static FName ExposedPropertyVariablesName = GET_MEMBER_NAME_CHECKED(FControlRigTraitSharedData, ExposedPropertyVariableNames);
-						static FName ExposedPropertyControlsName = GET_MEMBER_NAME_CHECKED(FControlRigTraitSharedData, ExposedPropertyControlNames);
-						static FName ExposedPropertyControlTypesName = GET_MEMBER_NAME_CHECKED(FControlRigTraitSharedData, ExposedPropertyControlTypes);
-						static FName ExposedPropertyControlDefaultValuesName = GET_MEMBER_NAME_CHECKED(FControlRigTraitSharedData, ExposedPropertyControlDefaultValues);
+						static FName ExposedPropertyVariablesName = GET_MEMBER_NAME_CHECKED(FControlRigTraitSharedDataBase, ExposedPropertyVariableNames);
+						static FName ExposedPropertyControlsName = GET_MEMBER_NAME_CHECKED(FControlRigTraitSharedDataBase, ExposedPropertyControlNames);
+						static FName ExposedPropertyControlTypesName = GET_MEMBER_NAME_CHECKED(FControlRigTraitSharedDataBase, ExposedPropertyControlTypes);
+						static FName ExposedPropertyControlDefaultValuesName = GET_MEMBER_NAME_CHECKED(FControlRigTraitSharedDataBase, ExposedPropertyControlDefaultValues);
 						
 						const FString ControlRigTraitName = FindControlRigTraitPinName(ModelNode);
 						ensureAlwaysMsgf(!ControlRigTraitName.IsEmpty(), TEXT("Can not Find a Control Rig Trait pin by name at node : [%s]."), *ModelNode->GetFName().ToString());
@@ -166,7 +170,18 @@ void FControlRigTraitSharedDataCustomization::OnPropertyExposeCheckboxChanged(EC
 							{
 								if (OptionalPin.bShowPin)
 								{
-									if (IsVariableProperty(ControlRigTraitSharedData, OptionalPin.PropertyName))
+									bool bIsVariableProperty = false;
+									if (ScopedControlRigTraitSharedData->GetStruct() == FControlRigTraitSharedData::StaticStruct())
+									{
+										FControlRigTraitSharedData* SharedData = (FControlRigTraitSharedData*)ScopedControlRigTraitSharedData->GetStructMemory();
+										bIsVariableProperty = IsVariableProperty(SharedData, OptionalPin.PropertyName);
+									}
+									else if (ScopedControlRigTraitSharedData->GetStruct() == FControlRigTraitSharedData_v2::StaticStruct())
+									{
+										FControlRigTraitSharedData_v2* SharedData = (FControlRigTraitSharedData_v2*)ScopedControlRigTraitSharedData->GetStructMemory();
+										bIsVariableProperty = IsVariableProperty(SharedData, OptionalPin.PropertyName);
+									}
+									if (bIsVariableProperty)
 									{
 										ControlRigTraitSharedData->ExposedPropertyVariableNames.Add(OptionalPin.PropertyName);
 									}
@@ -183,7 +198,7 @@ void FControlRigTraitSharedDataCustomization::OnPropertyExposeCheckboxChanged(EC
 								}
 							}
 
-							static constexpr const auto SetPinArrayDefaultValue = ([](URigVMController* Controller, URigVMPin* Pin, FControlRigTraitSharedData* ControlRigTraitSharedData, const FName& PropertyName, void* ArrayData) -> bool
+							static constexpr const auto SetPinArrayDefaultValue = ([](URigVMController* Controller, URigVMPin* Pin, FControlRigTraitSharedDataBase* ControlRigTraitSharedData, const FName& PropertyName, void* ArrayData) -> bool
 								{
 									if (const FProperty* Property = ControlRigTraitSharedData->StaticStruct()->FindPropertyByName(PropertyName))
 									{
@@ -271,18 +286,31 @@ TSharedPtr<FStructOnScope> FControlRigTraitSharedDataCustomization::GetControlRi
 
 FControlRigAssetStrongReference FControlRigTraitSharedDataCustomization::GetTargetAssetReference() const
 {
-	FControlRigTraitSharedData* ControlRigTraitSharedData = (FControlRigTraitSharedData*)ScopedControlRigTraitSharedData->GetStructMemory();
-	FControlRigAssetStrongReference AssetReference;
-	AssetReference.Set(ControlRigTraitSharedData->ControlRigClass);
-	return AssetReference;
+	if (ScopedControlRigTraitSharedData->GetStruct() == FControlRigTraitSharedData::StaticStruct())
+	{
+		FControlRigTraitSharedData* ControlRigTraitSharedData = (FControlRigTraitSharedData*)ScopedControlRigTraitSharedData->GetStructMemory();
+		return ControlRigTraitSharedData->GetControlRigAssetReference();
+	}
+	else if (ScopedControlRigTraitSharedData->GetStruct() == FControlRigTraitSharedData_v2::StaticStruct())
+	{
+		FControlRigTraitSharedData_v2* ControlRigTraitSharedData = (FControlRigTraitSharedData_v2*)ScopedControlRigTraitSharedData->GetStructMemory();
+		return ControlRigTraitSharedData->GetControlRigAssetReference();
+	}
+	return FControlRigAssetStrongReference();
 }
 
 USkeleton* FControlRigTraitSharedDataCustomization::GetTargetSkeleton() const
 {
 	USkeleton* TargetSkeleton = nullptr;
 
-	if (const FControlRigTraitSharedData* ControlRigTraitSharedData = (FControlRigTraitSharedData*)ScopedControlRigTraitSharedData->GetStructMemory())
+	if (ScopedControlRigTraitSharedData->GetStruct() == FControlRigTraitSharedData::StaticStruct())
 	{
+		FControlRigTraitSharedData* ControlRigTraitSharedData = (FControlRigTraitSharedData*)ScopedControlRigTraitSharedData->GetStructMemory();
+		TargetSkeleton = ControlRigTraitSharedData->GetPreviewSkeleton();
+	}
+	else if (ScopedControlRigTraitSharedData->GetStruct() == FControlRigTraitSharedData_v2::StaticStruct())
+	{
+		FControlRigTraitSharedData_v2* ControlRigTraitSharedData = (FControlRigTraitSharedData_v2*)ScopedControlRigTraitSharedData->GetStructMemory();
 		TargetSkeleton = ControlRigTraitSharedData->GetPreviewSkeleton();
 	}
 
@@ -299,15 +327,15 @@ const FControlRigIOMapping::FControlsInfo* FControlRigTraitSharedDataCustomizati
 	return ControlInfo;
 }
 
-bool FControlRigTraitSharedDataCustomization::IsVariableProperty(FControlRigTraitSharedData* ControlRigTraitSharedData, const FName& PropertyName )
+template<typename TSharedData>
+bool FControlRigTraitSharedDataCustomization::IsVariableProperty(TSharedData* ControlRigTraitSharedData, const FName& PropertyName )
 {
 	if (ControlRigTraitSharedData != nullptr)
 	{
-		UClass* ControlRigClass = ControlRigTraitSharedData->ControlRigClass;
-
-		if (const UControlRig* CDO = ControlRigClass->GetDefaultObject<UControlRig>())
+		FControlRigAssetStrongReference ControlRigClass = ControlRigTraitSharedData->GetControlRigAssetReference();
+		if (ControlRigClass.IsValid())
 		{
-			const TArray<FRigVMExternalVariable> PublicVariables = CDO->GetPublicVariables();
+			const TArray<FRigVMExternalVariable> PublicVariables = ControlRigClass.GetExternalVariables();
 
 			const FRigVMExternalVariable* PublicVar = PublicVariables.FindByPredicate([&PropertyName](const FRigVMExternalVariable& In)
 				{
@@ -330,86 +358,94 @@ void FControlRigTraitSharedDataCustomization::OnObjectsReinstanced(const TMap<UO
 			UObject* NewObject = Pair.Value;
 			if ((NewObject != nullptr) && NewObject->IsA<UControlRig>())
 			{
-				if (FControlRigTraitSharedData* ControlRigTraitSharedData = (FControlRigTraitSharedData*)ScopedControlRigTraitSharedData->GetStructMemory())
+				FControlRigAssetStrongReference ControlRigClass;
+				if (ScopedControlRigTraitSharedData->GetStruct() == FControlRigTraitSharedData::StaticStruct())
+				{
+					FControlRigTraitSharedData* ControlRigTraitSharedData = (FControlRigTraitSharedData*)ScopedControlRigTraitSharedData->GetStructMemory();
+					ControlRigClass = ControlRigTraitSharedData->GetControlRigAssetReference();
+				}
+				else if (ScopedControlRigTraitSharedData->GetStruct() == FControlRigTraitSharedData_v2::StaticStruct())
 				{
-					if (UClass* ControlRigClass = ControlRigTraitSharedData->ControlRigClass)
+					FControlRigTraitSharedData_v2* ControlRigTraitSharedData = (FControlRigTraitSharedData_v2*)ScopedControlRigTraitSharedData->GetStructMemory();
+					ControlRigClass = ControlRigTraitSharedData->GetControlRigAssetReference();
+				}
+				if (FControlRigTraitSharedDataBase* ControlRigTraitSharedData = (FControlRigTraitSharedDataBase*)ScopedControlRigTraitSharedData->GetStructMemory())
+				{
+					if (ControlRigClass.IsValid())
 					{
-						if (const UControlRig* CDO = ControlRigClass->GetDefaultObject<UControlRig>())
+						if (URigVMEdGraphNode* EdGraphNode = Cast<URigVMEdGraphNode>(SelectedNodeWeak.Get()))
 						{
-							if (URigVMEdGraphNode* EdGraphNode = Cast<URigVMEdGraphNode>(SelectedNodeWeak.Get()))
-							{
-								bool bRequiresPinRefresh = false;
+							bool bRequiresPinRefresh = false;
 
-								const TArray<FRigVMExternalVariable> PublicVariables = CDO->GetPublicVariables();
+							const TArray<FRigVMExternalVariable> PublicVariables = ControlRigClass.GetExternalVariables();
 
-								// Check if the currently exposed properties still exist
-								for (const FOptionalPinFromProperty& PinProperty : CustomPinProperties)
+							// Check if the currently exposed properties still exist
+							for (const FOptionalPinFromProperty& PinProperty : CustomPinProperties)
+							{
+								if (!PinProperty.bShowPin) // if not exposed, no need to test
 								{
-									if (!PinProperty.bShowPin) // if not exposed, no need to test
-									{
-										continue;
-									}
+									continue;
+								}
 
-									const FName& PropertyName = PinProperty.PropertyName;
-									const FRigVMExternalVariable* PublicVar = PublicVariables.FindByPredicate([&PropertyName](const FRigVMExternalVariable& In)
-										{
-											return In.GetName() == PropertyName;
-										});
-							
-									if (PublicVar != nullptr) // found a public var with the same name of the exposed pin
+								const FName& PropertyName = PinProperty.PropertyName;
+								const FRigVMExternalVariable* PublicVar = PublicVariables.FindByPredicate([&PropertyName](const FRigVMExternalVariable& In)
+									{
+										return In.GetName() == PropertyName;
+									});
+						
+								if (PublicVar != nullptr) // found a public var with the same name of the exposed pin
+								{
+									const TArray<URigVMPin*>& RootPins = EdGraphNode->GetInputPins();
+									for (const URigVMPin* RootPîn : RootPins)
 									{
-										const TArray<URigVMPin*>& RootPins = EdGraphNode->GetInputPins();
-										for (const URigVMPin* RootPîn : RootPins)
+										// Try to find a node pin with the same name
+										if (const URigVMPin*const* PinPtr = RootPîn->GetSubPins().FindByPredicate([&PropertyName](const URigVMPin* In)
+											{
+												return In->GetFName() == PropertyName;
+											}))
 										{
-											// Try to find a node pin with the same name
-											if (const URigVMPin*const* PinPtr = RootPîn->GetSubPins().FindByPredicate([&PropertyName](const URigVMPin* In)
-												{
-													return In->GetFName() == PropertyName;
-												}))
+											const URigVMPin* Pin = (*PinPtr);
+											if (Pin->GetCPPType() != PublicVar->GetExtendedCPPType())
 											{
-												const URigVMPin* Pin = (*PinPtr);
-												if (Pin->GetCPPType() != PublicVar->GetExtendedCPPType())
-												{
-													bRequiresPinRefresh = true;
-													break;
-												}
+												bRequiresPinRefresh = true;
+												break;
 											}
-											else
+										}
+										else
+										{
+											// No pin found but if it is exposed, refresh
+											if (ControlRigTraitSharedData->ExposedPropertyVariableNames.Contains(PropertyName) 
+												|| ControlRigTraitSharedData->ExposedPropertyControlNames.Contains(PropertyName))
 											{
-												// No pin found but if it is exposed, refresh
-												if (ControlRigTraitSharedData->ExposedPropertyVariableNames.Contains(PropertyName) 
-													|| ControlRigTraitSharedData->ExposedPropertyControlNames.Contains(PropertyName))
-												{
-													bRequiresPinRefresh = true;
-													break;
-												}
-												// else all ok, no pin, but it is not exposed
+												bRequiresPinRefresh = true;
+												break;
 											}
+											// else all ok, no pin, but it is not exposed
 										}
 									}
-									else
-									{
-										// No public var found, maybe a control ?
-										const TArray<FControlRigIOMapping::FControlsInfo>& Controls = ControlRigIOMapping->GetControls();
-										if(Controls.FindByPredicate([&PropertyName](const FControlRigIOMapping::FControlsInfo& In)
-											{
-												return In.Name == PropertyName;
-											}) == nullptr)
+								}
+								else
+								{
+									// No public var found, maybe a control ?
+									const TArray<FControlRigIOMapping::FControlsInfo>& Controls = ControlRigIOMapping->GetControls();
+									if(Controls.FindByPredicate([&PropertyName](const FControlRigIOMapping::FControlsInfo& In)
 										{
-											// No control found, request a rebuild
-											bRequiresPinRefresh = true;
-											break;
-										}
+											return In.Name == PropertyName;
+										}) == nullptr)
+									{
+										// No control found, request a rebuild
+										bRequiresPinRefresh = true;
+										break;
 									}
 								}
+							}
 
-								if (bRequiresPinRefresh)
+							if (bRequiresPinRefresh)
+							{
+								if (URigVMController* Controller = EdGraphNode->GetController())
 								{
-									if (URigVMController* Controller = EdGraphNode->GetController())
-									{
-										FRigVMControllerCompileBracketScope CompileScope(Controller); // avoid multiple recompilations and force a details refresh after recompile
-										Controller->RepopulatePinsOnNode(EdGraphNode->GetModelNode(), true, false, true); // Force pin refresh on node
-									}
+									FRigVMControllerCompileBracketScope CompileScope(Controller); // avoid multiple recompilations and force a details refresh after recompile
+									Controller->RepopulatePinsOnNode(EdGraphNode->GetModelNode(), true, false, true); // Force pin refresh on node
 								}
 							}
 						}
